ModelingComponents: Add support for enabling Complex Collision on DynamicMeshActors in UEditorModelingObjectsCreationAPI. Currently required for line traces to hit these objects, which many UI interactions depend on.

#rnx

diff --git a/Engine/Plugins/Experimental/MeshModelingToolsetExp/Source/ModelingComponentsEditorOnly/Private/EditorModelingObjectsCreationAPI.cpp b/Engine/Plugins/Experimental/MeshModelingToolsetExp/Source/ModelingComponentsEditorOnly/Private/EditorModelingObjectsCreationAPI.cpp
--- a/Engine/Plugins/Experimental/MeshModelingToolsetExp/Source/ModelingComponentsEditorOnly/Private/EditorModelingObjectsCreationAPI.cpp
+++ b/Engine/Plugins/Experimental/MeshModelingToolsetExp/Source/ModelingComponentsEditorOnly/Private/EditorModelingObjectsCreationAPI.cpp
@@ -235,6 +235,16 @@ FCreateMeshObjectResult UEditorModelingObjectsCreationAPI::CreateDynamicMeshActo
 		NewComponent->SetMaterial(k, ComponentMaterials[k]);
 	}
 
+	// configure collision
+	if (CreateMeshParams.bEnableCollision)
+	{
+		NewComponent->CollisionType = CreateMeshParams.CollisionMode;
+		// enable complex collision so that raycasts can hit this object
+		NewComponent->bEnableComplexCollision = true;
+		// force collision update
+		NewComponent->UpdateCollision(false);
+	}
+
 	NewActor->PostEditChange();
 
 	// emit result
diff --git a/Engine/Source/Runtime/GeometryFramework/Private/Components/DynamicMeshComponent.cpp b/Engine/Source/Runtime/GeometryFramework/Private/Components/DynamicMeshComponent.cpp
--- a/Engine/Source/Runtime/GeometryFramework/Private/Components/DynamicMeshComponent.cpp
+++ b/Engine/Source/Runtime/GeometryFramework/Private/Components/DynamicMeshComponent.cpp
@@ -203,11 +203,17 @@ void UDynamicMeshComponent::UpdateAutoCalculatedTangents()
 	{
 		GetDynamicMesh()->ProcessMesh([&](const FDynamicMesh3& Mesh)
 		{
-			const FDynamicMeshUVOverlay* UVOverlay = Mesh.Attributes()->PrimaryUV();
-			const FDynamicMeshNormalOverlay* NormalOverlay = Mesh.Attributes()->PrimaryNormals();
-			AutoCalculatedTangents.SetMesh(&Mesh);
-			AutoCalculatedTangents.ComputeTriVertexTangents(NormalOverlay, UVOverlay, FComputeTangentsOptions());
-			AutoCalculatedTangents.SetMesh(nullptr);
+			if (Mesh.HasAttributes())
+			{
+				const FDynamicMeshUVOverlay* UVOverlay = Mesh.Attributes()->PrimaryUV();
+				const FDynamicMeshNormalOverlay* NormalOverlay = Mesh.Attributes()->PrimaryNormals();
+				if (UVOverlay && NormalOverlay)
+				{
+					AutoCalculatedTangents.SetMesh(&Mesh);
+					AutoCalculatedTangents.ComputeTriVertexTangents(NormalOverlay, UVOverlay, FComputeTangentsOptions());
+					AutoCalculatedTangents.SetMesh(nullptr);
+				}
+			}
 		});
 
 		bAutoCalculatedTangentsValid = true;
diff --git a/Engine/Source/Runtime/GeometryFramework/Public/Components/DynamicMeshComponent.h b/Engine/Source/Runtime/GeometryFramework/Public/Components/DynamicMeshComponent.h
--- a/Engine/Source/Runtime/GeometryFramework/Public/Components/DynamicMeshComponent.h
+++ b/Engine/Source/Runtime/GeometryFramework/Public/Components/DynamicMeshComponent.h
@@ -433,7 +433,6 @@ protected:
 	//
 	// Physics APIs
 	//
-
 public:
 	virtual bool GetPhysicsTriMeshData(struct FTriMeshCollisionData* CollisionData, bool InUseAllTriData) override;
 	virtual bool ContainsPhysicsTriMeshData(bool InUseAllTriData) const override;
@@ -442,11 +441,20 @@ public:
 	virtual UBodySetup* GetBodySetup() override;
 	virtual void SetBodySetup(UBodySetup* NewSetup);
 
+	/**
+	 * Force an update of the Collision/Physics data for this Component.
+	 * @param bOnlyIfPending only update if a collision update is pending, ie the underlying DynamicMesh changed and bDeferCollisionUpdates is enabled
+	 */
+	virtual void UpdateCollision(bool bOnlyIfPending);
+
 	/** Type of Collision Geometry to use for this Mesh */
 	UPROPERTY(EditAnywhere, BlueprintReadWrite, Category = "Dynamic Mesh Component|Collision")
 	TEnumAsByte<enum ECollisionTraceFlag> CollisionType = ECollisionTraceFlag::CTF_UseSimpleAsComplex;
 
-	/** If true, current mesh will be used as Complex Collision source mesh. Independent of the CollisionType setting. */
+	/** 
+	 * If true, current mesh will be used as Complex Collision source mesh. 
+	 * This is independent of the CollisionType setting, ie, even if Complex collision is enabled, if this is false, then the Complex Collision mesh will be empty
+	 */
 	UPROPERTY(EditAnywhere, BlueprintReadWrite, Category = "Dynamic Mesh Component|Collision");
 	bool bEnableComplexCollision = false;
 
@@ -462,7 +470,6 @@ protected:
 	virtual void RebuildPhysicsData();
 
 	bool bCollisionUpdatePending = false;
-	virtual void UpdateCollision(bool bOnlyIfPending);
 
 protected:
 
