GeometryFramework: Add mesh pool API to DynamicMeshActor, so that Actor BPs can create/use temporary meshes internally without creating an enormous amount of garbage

#rnx

diff --git a/Engine/Source/Runtime/GeometryFramework/Private/DynamicMeshActor.cpp b/Engine/Source/Runtime/GeometryFramework/Private/DynamicMeshActor.cpp
--- a/Engine/Source/Runtime/GeometryFramework/Private/DynamicMeshActor.cpp
+++ b/Engine/Source/Runtime/GeometryFramework/Private/DynamicMeshActor.cpp
@@ -22,8 +22,64 @@ ADynamicMeshActor::ADynamicMeshActor(const FObjectInitializer& ObjectInitializer
 }
 
 
+UDynamicMeshPool* ADynamicMeshActor::GetComputeMeshPool()
+{
+	if (DynamicMeshPool == nullptr && bEnableComputeMeshPool)
+	{
+		DynamicMeshPool = NewObject<UDynamicMeshPool>();
+	}
+	return DynamicMeshPool;
+}
+
+
+UDynamicMesh* ADynamicMeshActor::AllocateComputeMesh()
+{
+	if (bEnableComputeMeshPool)
+	{
+		UDynamicMeshPool* UsePool = GetComputeMeshPool();
+		if (UsePool)
+		{
+			return UsePool->RequestMesh();
+		}
+	}
+
+	// if we could not return a pool mesh, allocate a new mesh that isn't owned by the pool
+	return NewObject<UDynamicMesh>(this);
+}
+
 
+bool ADynamicMeshActor::ReleaseComputeMesh(UDynamicMesh* Mesh)
+{
+	if (bEnableComputeMeshPool && Mesh)
+	{
+		UDynamicMeshPool* UsePool = GetComputeMeshPool();
+		if (UsePool != nullptr)
+		{
+			UsePool->ReturnMesh(Mesh);
+			return true;
+		}
+	}
+	return false;
+}
+
+
+void ADynamicMeshActor::ReleaseAllComputeMeshes()
+{
+	UDynamicMeshPool* UsePool = GetComputeMeshPool();
+	if (UsePool)
+	{
+		UsePool->ReturnAllMeshes();
+	}
+}
 
+void ADynamicMeshActor::FreeAllComputeMeshes()
+{
+	UDynamicMeshPool* UsePool = GetComputeMeshPool();
+	if (UsePool)
+	{
+		UsePool->FreeAllMeshes();
+	}
+}
 
 
 
diff --git a/Engine/Source/Runtime/GeometryFramework/Private/UDynamicMesh.cpp b/Engine/Source/Runtime/GeometryFramework/Private/UDynamicMesh.cpp
--- a/Engine/Source/Runtime/GeometryFramework/Private/UDynamicMesh.cpp
+++ b/Engine/Source/Runtime/GeometryFramework/Private/UDynamicMesh.cpp
@@ -523,7 +523,7 @@ void UDynamicMesh::ImportCustomProperties(const TCHAR* SourceText, FFeedbackCont
 		}
 
 		// if we got here we failed. Rather than produce an empty mesh, we generate a small cube
-		UE_LOG(LogTemp, Warning, TEXT("UDynamicMesh text-based property serialization incomplete, generating box as placeholder. Try increasing geometry.DynamicMesh.TextBasedDupeTriThreshold, or geometry.DynamicMesh.DupeStashTimeout."))
+		UE_LOG(LogGeometry, Warning, TEXT("UDynamicMesh text-based property serialization incomplete, generating box as placeholder. Try increasing geometry.DynamicMesh.TextBasedDupeTriThreshold, or geometry.DynamicMesh.DupeStashTimeout."))
 		FMinimalBoxMeshGenerator BoxGen;
 		BoxGen.Box = UE::Geometry::FOrientedBox3d(FVector3d::Zero(), 50.0 * FVector3d::One());
 		FDynamicMesh3 GenMesh(&BoxGen.Generate());
@@ -559,7 +559,7 @@ UDynamicMesh* UDynamicMeshPool::RequestMesh()
 	// This will allow them to be garbage-collected (eventually)
 	if (!ensure(AllCreatedMeshes.Num() < CVarDynamicMeshPoolMaxPoolSizeThreshold.GetValueOnGameThread()))
 	{
-		UE_LOG(LogTemp, Warning, TEXT("UDynamicMeshPool Threshold of %d Allocated Meshes exceeded! Releasing references to all current meshes and forcing a garbage collection."), CVarDynamicMeshPoolMaxPoolSizeThreshold.GetValueOnGameThread());
+		UE_LOG(LogGeometry, Warning, TEXT("UDynamicMeshPool Threshold of %d Allocated Meshes exceeded! Releasing references to all current meshes and forcing a garbage collection."), CVarDynamicMeshPoolMaxPoolSizeThreshold.GetValueOnGameThread());
 		AllCreatedMeshes.Reset();
 		GEngine->ForceGarbageCollection(true);
 	}
@@ -572,7 +572,7 @@ UDynamicMesh* UDynamicMeshPool::RequestMesh()
 
 void UDynamicMeshPool::ReturnMesh(UDynamicMesh* Mesh)
 {
-	if (ensure(Mesh))
+	if ( ensure(Mesh) && ensure(AllCreatedMeshes.Contains(Mesh)) )
 	{
 		Mesh->Reset();
 		if (ensure(CachedMeshes.Contains(Mesh) == false))
diff --git a/Engine/Source/Runtime/GeometryFramework/Public/DynamicMeshActor.h b/Engine/Source/Runtime/GeometryFramework/Public/DynamicMeshActor.h
--- a/Engine/Source/Runtime/GeometryFramework/Public/DynamicMeshActor.h
+++ b/Engine/Source/Runtime/GeometryFramework/Public/DynamicMeshActor.h
@@ -11,19 +11,58 @@
 
 /**
  * ADynamicMeshActor is an Actor that has a USimpleDynamicMeshComponent as it's RootObject.
- * Used for experimental procedural geometry generation Blueprints. Currently no physics support.
  */
 UCLASS(ConversionRoot, ComponentWrapperClass, ClassGroup=DynamicMesh, meta = (ChildCanTick))
 class GEOMETRYFRAMEWORK_API ADynamicMeshActor : public AActor
 {
 	GENERATED_UCLASS_BODY()
+
 private:
 	UPROPERTY(Category = DynamicMeshActor, VisibleAnywhere, BlueprintReadOnly, meta = (ExposeFunctionCategories = "Mesh,Rendering,Physics,Components|StaticMesh", AllowPrivateAccess = "true"))
 	TObjectPtr<class UDynamicMeshComponent> DynamicMeshComponent;
 
 public:
 	UFUNCTION(BlueprintCallable, Category = DynamicMeshActor)
-		UDynamicMeshComponent* GetDynamicMeshComponent() const { return DynamicMeshComponent; }
+	UDynamicMeshComponent* GetDynamicMeshComponent() const { return DynamicMeshComponent; }
+
+
+
+	//
+	// Mesh Pool support. Meshes can be locally allocated from the Mesh Pool
+	// in Blueprints, and then released back to the Pool and re-used. This
+	// avoids generating temporary UDynamicMesh instances that need to be
+	// garbage-collected. See UDynamicMeshPool for more details.
+	//
+
+public:
+	/** Control whether the DynamicMeshPool will be created when requested via GetComputeMeshPool() */
+	UPROPERTY(Category = DynamicMeshActor, VisibleAnywhere, BlueprintReadWrite)
+	bool bEnableComputeMeshPool = true;
+private:
+	/** The internal Mesh Pool, for use in DynamicMeshActor BPs. Use GetComputeMeshPool() to access this, as it will only be created on-demand if bEnableComputeMeshPool = true */
+	UPROPERTY(Transient)
+	TObjectPtr<UDynamicMeshPool> DynamicMeshPool;
+
+public:
+	/** Access the compute mesh pool */
+	UFUNCTION(BlueprintCallable, Category = DynamicMeshActor)
+	UDynamicMeshPool* GetComputeMeshPool();
+
+	/** Request a compute mesh from the Pool, which will return a previously-allocated mesh or add and return a new one. If the Pool is disabled, a new UDynamicMesh will be allocated and returned. */
+	UFUNCTION(BlueprintCallable, Category = DynamicMeshActor)
+	UDynamicMesh* AllocateComputeMesh();
+
+	/** Release a compute mesh back to the Pool */
+	UFUNCTION(BlueprintCallable, Category = DynamicMeshActor)
+	bool ReleaseComputeMesh(UDynamicMesh* Mesh);
+
+	/** Release all compute meshes that the Pool has allocated */
+	UFUNCTION(BlueprintCallable, Category = DynamicMeshActor)
+	void ReleaseAllComputeMeshes();
+
+	/** Release all compute meshes that the Pool has allocated, and then release them from the Pool, so that they will be garbage-collected */
+	UFUNCTION(BlueprintCallable, Category = DynamicMeshActor)
+	void FreeAllComputeMeshes();
 };
 
 
diff --git a/Engine/Source/Runtime/GeometryFramework/Public/UDynamicMesh.h b/Engine/Source/Runtime/GeometryFramework/Public/UDynamicMesh.h
--- a/Engine/Source/Runtime/GeometryFramework/Public/UDynamicMesh.h
+++ b/Engine/Source/Runtime/GeometryFramework/Public/UDynamicMesh.h
@@ -364,7 +364,7 @@ public:
  * the internal FDynamicMesh3 (which uses normal C++ memory management, so no garbage collection involved)
  * So the Pool does not re-use mesh memory, only the UObject containers.
  */
-UCLASS(Transient)
+UCLASS(BlueprintType, Transient)
 class GEOMETRYFRAMEWORK_API UDynamicMeshPool : public UObject
 {
 	GENERATED_BODY()
