GeometryFramework: add UObject flags for UBodySetup in UDynamicMeshComponent which resolve GLEO errors on save of DynamicMeshActor BPs

#rnx

diff --git a/Engine/Source/Runtime/GeometryFramework/Private/Components/DynamicMeshComponent.cpp b/Engine/Source/Runtime/GeometryFramework/Private/Components/DynamicMeshComponent.cpp
--- a/Engine/Source/Runtime/GeometryFramework/Private/Components/DynamicMeshComponent.cpp
+++ b/Engine/Source/Runtime/GeometryFramework/Private/Components/DynamicMeshComponent.cpp
@@ -60,6 +60,13 @@ void UDynamicMeshComponent::PostLoad()
 
 	ResetProxy();
 
+	// This is a fixup for existing UDynamicMeshComponents that did not have the correct flags 
+	// on the Instanced UBodySetup, these flags are now set in GetBodySetup() on new instances
+	if (MeshBodySetup && IsTemplate())
+	{
+		MeshBodySetup->SetFlags(RF_Public | RF_ArchetypeObject);
+	}
+
 	// make sure BodySetup is created
 	GetBodySetup();
 }
@@ -1074,7 +1081,11 @@ UBodySetup* UDynamicMeshComponent::GetBodySetup()
 		UBodySetup* NewBodySetup = nullptr;
 		{
 			FGCScopeGuard Scope;
-			NewBodySetup = NewObject<UBodySetup>(this);
+
+			// Below flags are copied from UProceduralMeshComponent::CreateBodySetupHelper(). Without these flags, DynamicMeshComponents inside
+			// a DynamicMeshActor BP will result on a GLEO error after loading and modifying a saved Level (but *not* on the initial save)
+			// The UBodySetup in a template needs to be public since the property is Instanced and thus is the archetype of the instance meaning there is a direct reference
+			NewBodySetup = NewObject<UBodySetup>(this, NAME_None, (IsTemplate() ? RF_Public | RF_ArchetypeObject : RF_NoFlags));
 		}
 		NewBodySetup->BodySetupGuid = FGuid::NewGuid();
 		
