GeometryFramework:
- UDynamicMesh Reset() functions now return themselves, configured as BP output pin, helps to make BP wiring cleaner
- added DynamicMeshComponent::ConfigureMaterialSet() and ::ValidateMaterials(), to allow material set to be updated from BP
- added BaseDynamicMeshComponent::SetNumMaterials() to support above
#rnx

diff --git a/Engine/Source/Runtime/GeometryFramework/Private/Components/BaseDynamicMeshComponent.cpp b/Engine/Source/Runtime/GeometryFramework/Private/Components/BaseDynamicMeshComponent.cpp
--- a/Engine/Source/Runtime/GeometryFramework/Private/Components/BaseDynamicMeshComponent.cpp
+++ b/Engine/Source/Runtime/GeometryFramework/Private/Components/BaseDynamicMeshComponent.cpp
@@ -124,6 +124,23 @@ void UBaseDynamicMeshComponent::SetMaterial(int32 ElementIndex, UMaterialInterfa
 }
 
 
+void UBaseDynamicMeshComponent::SetNumMaterials(int32 NumMaterials)
+{
+	if (BaseMaterials.Num() > NumMaterials)
+	{
+		// discard extra materials
+		BaseMaterials.SetNum(NumMaterials);
+	}
+	else
+	{
+		while (NumMaterials < BaseMaterials.Num())
+		{
+			SetMaterial(NumMaterials, nullptr);
+			NumMaterials++;
+		}
+	}
+}
+
 void UBaseDynamicMeshComponent::GetUsedMaterials(TArray<UMaterialInterface*>& OutMaterials, bool bGetDebugMaterials) const
 {
 	UMeshComponent::GetUsedMaterials(OutMaterials, bGetDebugMaterials);
diff --git a/Engine/Source/Runtime/GeometryFramework/Private/Components/DynamicMeshComponent.cpp b/Engine/Source/Runtime/GeometryFramework/Private/Components/DynamicMeshComponent.cpp
--- a/Engine/Source/Runtime/GeometryFramework/Private/Components/DynamicMeshComponent.cpp
+++ b/Engine/Source/Runtime/GeometryFramework/Private/Components/DynamicMeshComponent.cpp
@@ -180,6 +180,51 @@ void UDynamicMeshComponent::ApplyTransform(const UE::Geometry::FTransform3d& Tra
 
 
 
+bool UDynamicMeshComponent::ValidateMaterialSlots(bool bCreateIfMissing, bool bDeleteExtraSlots)
+{
+	int32 MaxMeshMaterialID = 0;
+	MeshObject->ProcessMesh([&](const FDynamicMesh3& EditMesh)
+	{
+		if (EditMesh.HasAttributes() && EditMesh.Attributes()->HasMaterialID() && EditMesh.Attributes()->GetMaterialID() != nullptr)
+		{
+			const FDynamicMeshMaterialAttribute* MaterialIDs = EditMesh.Attributes()->GetMaterialID();
+			for (int TriangleID : EditMesh.TriangleIndicesItr())
+			{
+				MaxMeshMaterialID = FMath::Max(MaxMeshMaterialID, MaterialIDs->GetValue(TriangleID));
+			}
+		}
+	});
+	int32 NumRequiredMaterials = MaxMeshMaterialID + 1;
+
+	int32 NumMaterials = GetNumMaterials();
+	if ( bCreateIfMissing && NumMaterials < NumRequiredMaterials )
+	{
+		for (int32 k = NumMaterials; k < NumRequiredMaterials; ++k)
+		{
+			SetMaterial(k, nullptr);
+		}
+	}
+	NumMaterials = GetNumMaterials();
+
+	if (bDeleteExtraSlots && NumMaterials > NumRequiredMaterials)
+	{
+		SetNumMaterials(NumRequiredMaterials);
+	}
+	NumMaterials = GetNumMaterials();
+
+	return (NumMaterials == NumRequiredMaterials);
+}
+
+
+void UDynamicMeshComponent::ConfigureMaterialSet(const TArray<UMaterialInterface*>& NewMaterialSet)
+{
+	for (int k = 0; k < NewMaterialSet.Num(); ++k)
+	{
+		SetMaterial(k, NewMaterialSet[k]);
+	}	
+}
+
+
 void UDynamicMeshComponent::SetTangentsType(EDynamicMeshComponentTangentsMode NewTangentsType)
 {
 	if (NewTangentsType != TangentsType)
diff --git a/Engine/Source/Runtime/GeometryFramework/Private/UDynamicMesh.cpp b/Engine/Source/Runtime/GeometryFramework/Private/UDynamicMesh.cpp
--- a/Engine/Source/Runtime/GeometryFramework/Private/UDynamicMesh.cpp
+++ b/Engine/Source/Runtime/GeometryFramework/Private/UDynamicMesh.cpp
@@ -42,7 +42,7 @@ void UDynamicMesh::InitializeNewMesh()
 	InitializeMesh();
 }
 
-void UDynamicMesh::Reset()
+UDynamicMesh* UDynamicMesh::Reset()
 {
 	FDynamicMeshChangeInfo ChangeInfo;
 	ChangeInfo.Type = EDynamicMeshChangeType::GeneralEdit;
@@ -51,9 +51,10 @@ void UDynamicMesh::Reset()
 		check(&EditMesh == Mesh.Get());		// assuming that EditMesh is internal mesh here...
 		InitializeMesh();
 	}, ChangeInfo);
+	return this;
 }
 
-void UDynamicMesh::ResetToCube()
+UDynamicMesh* UDynamicMesh::ResetToCube()
 {
 	FDynamicMeshChangeInfo ChangeInfo;
 	ChangeInfo.Type = EDynamicMeshChangeType::GeneralEdit;
@@ -66,6 +67,7 @@ void UDynamicMesh::ResetToCube()
 		EditMesh.EnableTriangleGroups();
 		EditMesh.Attributes()->EnableMaterialID();
 	}, ChangeInfo);
+	return this;
 }
 
 
diff --git a/Engine/Source/Runtime/GeometryFramework/Public/Components/BaseDynamicMeshComponent.h b/Engine/Source/Runtime/GeometryFramework/Public/Components/BaseDynamicMeshComponent.h
--- a/Engine/Source/Runtime/GeometryFramework/Public/Components/BaseDynamicMeshComponent.h
+++ b/Engine/Source/Runtime/GeometryFramework/Public/Components/BaseDynamicMeshComponent.h
@@ -334,6 +334,8 @@ public:
 	virtual void SetMaterial(int32 ElementIndex, UMaterialInterface* Material) override;
 	virtual void GetUsedMaterials(TArray<UMaterialInterface*>& OutMaterials, bool bGetDebugMaterials = false) const override;
 
+	virtual void SetNumMaterials(int32 NumMaterials);
+
 	UPROPERTY()
 	TArray<TObjectPtr<UMaterialInterface>> BaseMaterials;
 };
diff --git a/Engine/Source/Runtime/GeometryFramework/Public/Components/DynamicMeshComponent.h b/Engine/Source/Runtime/GeometryFramework/Public/Components/DynamicMeshComponent.h
--- a/Engine/Source/Runtime/GeometryFramework/Public/Components/DynamicMeshComponent.h
+++ b/Engine/Source/Runtime/GeometryFramework/Public/Components/DynamicMeshComponent.h
@@ -390,6 +390,28 @@ public:
 
 
 
+	//===============================================================================================================
+	// Material Set API. DynamicMeshComponent supports changing the Material Set dynamically, even at Runtime.
+public:
+
+	/**
+	 * Set new list of Materials for the Mesh. Dynamic Mesh Component does not have 
+	 * Slot Names, so the size of the Material Set should be the same as the number of
+	 * different Material IDs on the mesh MaterialID attribute
+	 */
+	UFUNCTION(BlueprintCallable, Category = "Dynamic Mesh Component")
+	void ConfigureMaterialSet(const TArray<UMaterialInterface*>& NewMaterialSet);
+
+	/**
+	 * Compute the maximum MaterialID on the DynamicMesh, and ensure that Material Slots match.
+	 * Pass both arguments as false to just do a check.
+	 * @param bCreateIfMissing if true, add extra (empty) Material Slots to match max MaterialID
+	 * @param bDeleteExtraSlots if true, extra Material Slots beyond max MaterialID are removed
+	 * @return true if at the end of this function, Material Slot Count == Max MaterialID
+	 */
+	UFUNCTION(BlueprintCallable, Category = "Dynamic Mesh Component")
+	bool ValidateMaterialSlots(bool bCreateIfMissing = true, bool bDeleteExtraSlots = true);
+
 
 	//===============================================================================================================
 	// Triangle-Vertex Tangents support. The default behavior is to not use Tangents, this will lead to incorrect
diff --git a/Engine/Source/Runtime/GeometryFramework/Public/UDynamicMesh.h b/Engine/Source/Runtime/GeometryFramework/Public/UDynamicMesh.h
--- a/Engine/Source/Runtime/GeometryFramework/Public/UDynamicMesh.h
+++ b/Engine/Source/Runtime/GeometryFramework/Public/UDynamicMesh.h
@@ -110,14 +110,16 @@ public:
 	 * This *does not* allocate a new mesh, so any existing mesh pointers/refs are still valid
 	 */
 	UFUNCTION(BlueprintCallable, Category = "Dynamic Mesh")
-	void Reset();
+	UPARAM(DisplayName = "Self") UDynamicMesh* 
+	Reset();
 
 	/**
 	 * Clear the internal mesh to a 100x100x100 cube with base at the origin.
 	 * This this instead of Reset() if an initially-empty mesh is undesirable (eg for a Component)
 	 */
 	UFUNCTION(BlueprintCallable, Category = "Dynamic Mesh")
-	void ResetToCube();
+	UPARAM(DisplayName = "Self") UDynamicMesh* 
+	ResetToCube();
 
 	//
 	// Native access/modification functions
