GeometryFramework: add BP-exposed UDynamicMeshComponent::EnableComplexAsSimpleCollision() function
#rnx

diff --git a/Engine/Source/Runtime/GeometryFramework/Private/Components/DynamicMeshComponent.cpp b/Engine/Source/Runtime/GeometryFramework/Private/Components/DynamicMeshComponent.cpp
--- a/Engine/Source/Runtime/GeometryFramework/Private/Components/DynamicMeshComponent.cpp
+++ b/Engine/Source/Runtime/GeometryFramework/Private/Components/DynamicMeshComponent.cpp
@@ -1189,4 +1189,11 @@ void UDynamicMeshComponent::UpdateCollision(bool bOnlyIfPending)
 	{
 		RebuildPhysicsData();
 	}
+}
+
+void UDynamicMeshComponent::EnableComplexAsSimpleCollision()
+{
+	bEnableComplexCollision = true;
+	CollisionType = ECollisionTraceFlag::CTF_UseComplexAsSimple;
+	UpdateCollision(true);
 }
\ No newline at end of file
diff --git a/Engine/Source/Runtime/GeometryFramework/Public/Components/DynamicMeshComponent.h b/Engine/Source/Runtime/GeometryFramework/Public/Components/DynamicMeshComponent.h
--- a/Engine/Source/Runtime/GeometryFramework/Public/Components/DynamicMeshComponent.h
+++ b/Engine/Source/Runtime/GeometryFramework/Public/Components/DynamicMeshComponent.h
@@ -456,6 +456,9 @@ protected:
 	// Physics APIs
 	//
 public:
+	UFUNCTION(BlueprintCallable, Category = "Dynamic Mesh Component")
+	void EnableComplexAsSimpleCollision();
+
 	virtual bool GetPhysicsTriMeshData(struct FTriMeshCollisionData* CollisionData, bool InUseAllTriData) override;
 	virtual bool ContainsPhysicsTriMeshData(bool InUseAllTriData) const override;
 	virtual bool WantsNegXTriMesh() override;
