GeometryFramework: add support for in-Editor deferred mesh regeneration for BP procedural mesh generators built with DynamicMeshActor
#rnx

diff --git a/Engine/Source/Runtime/GeometryFramework/Private/DynamicMeshActor.cpp b/Engine/Source/Runtime/GeometryFramework/Private/DynamicMeshActor.cpp
--- a/Engine/Source/Runtime/GeometryFramework/Private/DynamicMeshActor.cpp
+++ b/Engine/Source/Runtime/GeometryFramework/Private/DynamicMeshActor.cpp
@@ -19,6 +19,16 @@ ADynamicMeshActor::ADynamicMeshActor(const FObjectInitializer& ObjectInitializer
 	DynamicMeshComponent->SetMaterial(0, UMaterial::GetDefaultMaterial(MD_Surface));		// is this necessary?
 
 	SetRootComponent(DynamicMeshComponent);
+
+	//
+	// Configure ADynamicMeshActor to always Tick(). This is necessary for the bIsEditorGeneratedMeshActor 
+	// Tick-in-Editor support below to work w/o any other user intervention. Generally always Ticking is bad behavior, 
+	// however DynamicMeshActor is inherently expensive and so we do not expect that large numbers of them
+	// will be used in a performance-sensitive context.
+	//
+	PrimaryActorTick.bCanEverTick = true;
+	PrimaryActorTick.bStartWithTickEnabled = true;
+	PrimaryActorTick.SetTickFunctionEnable(true);
 }
 
 
@@ -83,4 +93,49 @@ void ADynamicMeshActor::FreeAllComputeMeshes()
 
 
 
+
+
+
+void ADynamicMeshActor::Tick(float DeltaTime)
+{
+	Super::Tick(DeltaTime);
+
+#if WITH_EDITOR
+	if (bIsEditorGeneratedMeshActor && bGeneratedMeshRebuildPending)
+	{
+		OnEditorRebuildGeneratedMesh();
+		bGeneratedMeshRebuildPending = false;
+	}
+#endif
+}
+
+
+
+void ADynamicMeshActor::OnConstruction(const FTransform& Transform)
+{
+	Super::OnConstruction(Transform);
+
+#if WITH_EDITOR
+	if (bIsEditorGeneratedMeshActor)
+	{
+		bGeneratedMeshRebuildPending = true;
+	}
+#endif
+}
+
+
+#if WITH_EDITOR
+bool ADynamicMeshActor::ShouldTickIfViewportsOnly() const
+{
+	if (bIsEditorGeneratedMeshActor && GetWorld() != nullptr && GetWorld()->WorldType == EWorldType::Editor)
+	{
+		return true;
+	}
+	return false;
+}
+#endif
+
+
+
+
 #undef LOCTEXT_NAMESPACE
\ No newline at end of file
diff --git a/Engine/Source/Runtime/GeometryFramework/Public/DynamicMeshActor.h b/Engine/Source/Runtime/GeometryFramework/Public/DynamicMeshActor.h
--- a/Engine/Source/Runtime/GeometryFramework/Public/DynamicMeshActor.h
+++ b/Engine/Source/Runtime/GeometryFramework/Public/DynamicMeshActor.h
@@ -63,6 +63,73 @@ public:
 	/** Release all compute meshes that the Pool has allocated, and then release them from the Pool, so that they will be garbage-collected */
 	UFUNCTION(BlueprintCallable, Category = DynamicMeshActor)
 	void FreeAllComputeMeshes();
+
+
+
+
+	//
+	// In-Editor Generated/Procedural Mesh Actor support. 
+	// 
+	// These members are used to help with the creation of DynamicMeshActors that generate
+	// meshes dynamically in the Editor, for example via Blueprints. Expensive procedural 
+	// generation via BP can potentially cause major problems in the Editor, and so the
+	// implmentor of a BP Actor mesh generator needs to take some care.
+	// 
+	// If bIsEditorGeneratedMeshActor = true, then OnEditorRebuildGeneratedMesh
+	// will fire after the Construction script runs for this Actor. OnEditorRebuildGeneratedMesh does not directly
+	// fire from the Construction Script, but rather in the next Tick. This works around a fundamental
+	// limitation related to "mouse event priority" and slider/transform input in the UE Editor.
+	// Essentially doing even a moderately expensive mesh generation operation directly in the 
+	// Construction Script will appear to freeze the Viewport, while doing the same operation in the
+	// OnEditorRebuildGeneratedMesh implementation will remain responsive (albeit at a low FPS rate).
+	// 
+	// In future this may be extended to (eg) throttle procedural regeneration to further improve
+	// interactivity. 
+	// 
+	// Note that the current implementation requires in-Editor Ticking for this Actor. If bIsEditorGeneratedMeshActor
+	// is true, then ShouldTickIfViewportsOnly() will return true, which means that the Actor Tick Event
+	// will also fire in Editor viewports. This means that any game logic wired to the Tick Event
+	// in the Actor BP will aso run outside of PIE!
+	// 
+	// Note that the above behavior currently only occurs in the Editor, where the Construction script 
+	// is frequently re-executed (on Actor Transform and Property Editing). 
+	//
+public:
+
+	/**
+	 * Set this flag to true in a BP Subclass of DynamicMeshActor that procedurally generates it's
+	 * own mesh via BP (eg with Geometry Script). This will cause OnEditorRebuildGeneratedMesh to fire after 
+	 * the Construction Script is run or re-run. Regenerating the mesh on this event, rather than directly
+	 * in the Construction Script, will result in better interactive performance in the Editor.
+	 */
+	UPROPERTY(EditAnywhere, Category = DynamicMeshActor )
+	bool bIsEditorGeneratedMeshActor = false;
+
+	/**
+	 * This event will fire from Tick() to notify listeners that the internal Mesh should
+	 * be rebuilt. Procedural Mesh Generator Actors built in BP's derived from DynamicMeshActor
+	 * should listen for this event and use it to rebuild their meshes, instead of doing so
+	 * directly from the Construction Script.
+	 */
+	UFUNCTION(BlueprintImplementableEvent, CallInEditor, Category = "Events")
+	void OnEditorRebuildGeneratedMesh();
+
+
+
+	// the functions below are overridden to implement the above behavior
+public:
+	virtual void Tick(float DeltaTime) override;
+
+#if WITH_EDITOR
+	virtual bool ShouldTickIfViewportsOnly() const override;
+#endif
+
+	virtual void OnConstruction(const FTransform& Transform) override;
+
+protected:
+	// this internal flag is set in OnConstruction if bIsEditorGeneratedMeshActor=true, and will cause
+	// OnEditorRebuildGeneratedMesh to fire in Tick(), after which the flag is cleared
+	bool bGeneratedMeshRebuildPending = false;
 };
 
 
