GeometryScripting: move in-editor generated-mesh functionality from DynamicMeshActor to a subclass GeneratedDynamicMeshActor in experimental GeometryScriptingEditor module. Add EditorGeometryGenerationSubsystem to provide tick for mesh regeneration, instead of relying on Actor tick-in-Editor, which can be throttled by slate interaction (and is generally problematic).

#rnx

diff --git a/Engine/Plugins/Experimental/GeometryScripting/GeometryScripting.uplugin b/Engine/Plugins/Experimental/GeometryScripting/GeometryScripting.uplugin
--- a/Engine/Plugins/Experimental/GeometryScripting/GeometryScripting.uplugin
+++ b/Engine/Plugins/Experimental/GeometryScripting/GeometryScripting.uplugin
@@ -4,7 +4,7 @@
 	"VersionName": "0.1",
 	"FriendlyName": "Geometry Scripting",
 	"Description": "Geometry Scripting Modules enable procedural generation and manipulation of meshes/geometry via Blueprints",
-	"Category": "Other",
+	"Category": "Geometry",
 	"CreatedBy": "Epic Games, Inc.",
 	"CreatedByURL": "http://epicgames.com",
 	"DocsURL": "",
diff --git a/Engine/Plugins/Experimental/GeometryScripting/Source/GeometryScriptingEditor/GeometryScriptingEditor.Build.cs b/Engine/Plugins/Experimental/GeometryScripting/Source/GeometryScriptingEditor/GeometryScriptingEditor.Build.cs
--- a/Engine/Plugins/Experimental/GeometryScripting/Source/GeometryScriptingEditor/GeometryScriptingEditor.Build.cs
+++ b/Engine/Plugins/Experimental/GeometryScripting/Source/GeometryScriptingEditor/GeometryScriptingEditor.Build.cs
@@ -26,12 +26,14 @@ public class GeometryScriptingEditor : ModuleRules
 			new string[]
 			{
 				"Core",
+				"CoreUObject",
 				"PhysicsCore",
 				"RenderCore",
 				"GeometryCore",
 				"GeometryFramework",
 				"DynamicMesh",
-				"GeometryScriptingCore"
+				"GeometryScriptingCore",
+				"EditorSubsystem"
 			}
 			);
 
@@ -39,7 +41,6 @@ public class GeometryScriptingEditor : ModuleRules
 		PrivateDependencyModuleNames.AddRange(
 			new string[]
 			{
-				"CoreUObject",
 				"Engine",
 				"MeshDescription",
 				"StaticMeshDescription",
diff --git a/Engine/Plugins/Experimental/GeometryScripting/Source/GeometryScriptingEditor/Private/GeometryActors/EditorGeometryGenerationSubsystem.cpp b/Engine/Plugins/Experimental/GeometryScripting/Source/GeometryScriptingEditor/Private/GeometryActors/EditorGeometryGenerationSubsystem.cpp
new file mode 100644
--- /dev/null
+++ b/Engine/Plugins/Experimental/GeometryScripting/Source/GeometryScriptingEditor/Private/GeometryActors/EditorGeometryGenerationSubsystem.cpp
@@ -0,0 +1,111 @@
+// Copyright Epic Games, Inc. All Rights Reserved.
+
+#include "GeometryActors/EditorGeometryGenerationSubsystem.h"
+#include "GeometryActors/GeneratedDynamicMeshActor.h"
+#include "Editor.h"
+
+
+void UEditorGeometryGenerationManager::Tick(float DeltaTime)
+{
+	// rebuild all valid generated actors, if necessary
+	for (AGeneratedDynamicMeshActor* Actor : ActiveGeneratedActors)
+	{
+		if (Actor->IsValidLowLevel() && !Actor->IsPendingKillOrUnreachable() &&  Actor->GetLevel() != nullptr)
+		{
+			Actor->ExecuteRebuildGeneratedMeshIfPending();
+		}
+	}
+}
+
+bool UEditorGeometryGenerationManager::IsTickable() const
+{
+	return (ActiveGeneratedActors.Num() > 0);
+}
+
+TStatId UEditorGeometryGenerationManager::GetStatId() const
+{
+	RETURN_QUICK_DECLARE_CYCLE_STAT(UEditorGeometryGenerationManager, STATGROUP_Tickables);
+}
+
+
+void UEditorGeometryGenerationManager::Shutdown()
+{
+	ActiveGeneratedActors.Reset();
+}
+
+void UEditorGeometryGenerationManager::RegisterGeneratedMeshActor(AGeneratedDynamicMeshActor* Actor)
+{
+	ensure(ActiveGeneratedActors.Contains(Actor) == false);
+	ActiveGeneratedActors.Add(Actor);
+}
+
+void UEditorGeometryGenerationManager::UnregisterGeneratedMeshActor(AGeneratedDynamicMeshActor* Actor)
+{
+	ensure(ActiveGeneratedActors.Contains(Actor));
+	ActiveGeneratedActors.Remove(Actor);
+}
+
+
+
+
+bool UEditorGeometryGenerationSubsystem::bIsShuttingDown = false;
+
+
+void UEditorGeometryGenerationSubsystem::Initialize(FSubsystemCollectionBase& Collection)
+{
+	GenerationManager = NewObject<UEditorGeometryGenerationManager>(this);
+
+	GEditor->OnEditorClose().AddUObject(this, &UEditorGeometryGenerationSubsystem::OnShutdown);
+	FCoreDelegates::OnEnginePreExit.AddUObject(this, &UEditorGeometryGenerationSubsystem::OnShutdown);
+}
+
+void UEditorGeometryGenerationSubsystem::Deinitialize()
+{
+	GEditor->OnEditorClose().RemoveAll(this);
+	FCoreDelegates::OnEnginePreExit.RemoveAll(this);
+	bIsShuttingDown = true;
+
+	if (GenerationManager)
+	{
+		GenerationManager->Shutdown();
+		GenerationManager = nullptr;
+	}
+}
+
+
+void UEditorGeometryGenerationSubsystem::OnShutdown()
+{
+	bIsShuttingDown = true;
+}
+
+
+
+void UEditorGeometryGenerationSubsystem::RegisterGeneratedMeshActor(AGeneratedDynamicMeshActor* Actor)
+{
+	if (bIsShuttingDown || GEditor == nullptr)		// subsystem no longer exists
+	{
+		return;
+	}
+
+	UEditorGeometryGenerationSubsystem* Subsystem = GEditor->GetEditorSubsystem<UEditorGeometryGenerationSubsystem>();
+	if (ensure(Subsystem))
+	{
+		Subsystem->GenerationManager->RegisterGeneratedMeshActor(Actor);
+	}
+}
+
+void UEditorGeometryGenerationSubsystem::UnregisterGeneratedMeshActor(AGeneratedDynamicMeshActor* Actor)
+{
+	if (bIsShuttingDown || GEditor == nullptr)		// subsystem no longer exists
+	{
+		return;
+	}
+
+	UEditorGeometryGenerationSubsystem* Subsystem = GEditor->GetEditorSubsystem<UEditorGeometryGenerationSubsystem>();
+	if (ensure(Subsystem))
+	{
+		Subsystem->GenerationManager->UnregisterGeneratedMeshActor(Actor);
+	}
+}
+
+
diff --git a/Engine/Plugins/Experimental/GeometryScripting/Source/GeometryScriptingEditor/Private/GeometryActors/GeneratedDynamicMeshActor.cpp b/Engine/Plugins/Experimental/GeometryScripting/Source/GeometryScriptingEditor/Private/GeometryActors/GeneratedDynamicMeshActor.cpp
new file mode 100644
--- /dev/null
+++ b/Engine/Plugins/Experimental/GeometryScripting/Source/GeometryScriptingEditor/Private/GeometryActors/GeneratedDynamicMeshActor.cpp
@@ -0,0 +1,78 @@
+// Copyright Epic Games, Inc. All Rights Reserved.
+
+#include "GeometryActors/GeneratedDynamicMeshActor.h"
+#include "GeometryActors/EditorGeometryGenerationSubsystem.h"
+
+#define LOCTEXT_NAMESPACE "AGeneratedDynamicMeshActor"
+
+
+AGeneratedDynamicMeshActor::AGeneratedDynamicMeshActor(const FObjectInitializer& ObjectInitializer)
+	: Super(ObjectInitializer)
+{
+	RegisterWithGenerationManager();
+}
+
+AGeneratedDynamicMeshActor::~AGeneratedDynamicMeshActor()
+{
+	UnregisterWithGenerationManager();
+}
+
+
+void AGeneratedDynamicMeshActor::OnConstruction(const FTransform& Transform)
+{
+	Super::OnConstruction(Transform);
+
+	bGeneratedMeshRebuildPending = true;
+}
+
+
+void AGeneratedDynamicMeshActor::Destroyed()
+{
+	UnregisterWithGenerationManager();
+	Super::Destroyed();
+}
+
+
+void AGeneratedDynamicMeshActor::RegisterWithGenerationManager()
+{
+	if (HasAnyFlags(RF_ClassDefaultObject))		// ignore for CDO
+	{
+		return;
+	}
+
+	if (bIsRegisteredWithGenerationManager == false)
+	{
+		UEditorGeometryGenerationSubsystem::RegisterGeneratedMeshActor(this);
+		bIsRegisteredWithGenerationManager = true;
+	}
+}
+
+void AGeneratedDynamicMeshActor::UnregisterWithGenerationManager()
+{
+	if (HasAnyFlags(RF_ClassDefaultObject))		// ignore for CDO
+	{
+		return;
+	}
+
+	if (bIsRegisteredWithGenerationManager)
+	{
+		UEditorGeometryGenerationSubsystem::UnregisterGeneratedMeshActor(this);
+		bIsRegisteredWithGenerationManager = false;
+		bGeneratedMeshRebuildPending = false;
+	}
+}
+
+
+
+void AGeneratedDynamicMeshActor::ExecuteRebuildGeneratedMeshIfPending()
+{
+	if (bGeneratedMeshRebuildPending)
+	{
+		FEditorScriptExecutionGuard Guard;
+		OnRebuildGeneratedMesh();
+		bGeneratedMeshRebuildPending = false;
+	}
+}
+
+
+#undef LOCTEXT_NAMESPACE
\ No newline at end of file
diff --git a/Engine/Plugins/Experimental/GeometryScripting/Source/GeometryScriptingEditor/Public/GeometryActors/EditorGeometryGenerationSubsystem.h b/Engine/Plugins/Experimental/GeometryScripting/Source/GeometryScriptingEditor/Public/GeometryActors/EditorGeometryGenerationSubsystem.h
new file mode 100644
--- /dev/null
+++ b/Engine/Plugins/Experimental/GeometryScripting/Source/GeometryScriptingEditor/Public/GeometryActors/EditorGeometryGenerationSubsystem.h
@@ -0,0 +1,84 @@
+// Copyright Epic Games, Inc. All Rights Reserved.
+
+#pragma once
+
+#include "EditorSubsystem.h"
+#include "TickableEditorObject.h"
+
+#include "EditorGeometryGenerationSubsystem.generated.h"
+
+class AGeneratedDynamicMeshActor;
+
+
+/**
+ * UEditorGeometryGenerationSubsystem manages recomputation of "generated" mesh actors, eg
+ * to provide procedural mesh generation in-Editor. Generally such procedural mesh generation
+ * is expensive, and if many objects need to be generated, the regeneration needs to be 
+ * managed at a higher level to ensure that the Editor remains responsive/interactive.
+ * 
+ * AGeneratedDynamicMeshActors register themselves with this Subsystem, and
+ * allow the Subsystem to tell them when they should regenerate themselves (if necessary).
+ * The current behavior is to run all pending generations on a Tick, however in future
+ * this regeneration will be more carefully managed via throttling / timeslicing / etc.
+ * 
+ */
+UCLASS()
+class UEditorGeometryGenerationSubsystem : public UEditorSubsystem
+{
+	GENERATED_BODY()
+
+public:
+	virtual void Initialize(FSubsystemCollectionBase& Collection) override;
+	virtual void Deinitialize() override;
+
+public:
+	UPROPERTY()
+	TObjectPtr<UEditorGeometryGenerationManager> GenerationManager;
+
+
+	//
+	// Static functions that simplify registration of an Actor w/ the Subsystem
+	//
+public:
+	static void RegisterGeneratedMeshActor(AGeneratedDynamicMeshActor* Actor);
+	static void UnregisterGeneratedMeshActor(AGeneratedDynamicMeshActor* Actor);
+
+
+protected:
+	// callback connected to engine/editor shutdown events to set bIsShuttingDown, which disables the subsystem static functions above
+	virtual void OnShutdown();
+
+private:
+	static bool bIsShuttingDown;
+};
+
+
+
+
+/**
+ * UEditorGeometryGenerationManager is a class used by UEditorGeometryGenerationSubsystem to
+ * store registrations and provide a Tick()
+ */
+UCLASS()
+class UEditorGeometryGenerationManager : public UObject, public FTickableEditorObject
+{
+	GENERATED_BODY()
+
+public:
+	virtual void Shutdown();
+
+	virtual void RegisterGeneratedMeshActor(AGeneratedDynamicMeshActor* Actor);
+	virtual void UnregisterGeneratedMeshActor(AGeneratedDynamicMeshActor* Actor);
+
+public:
+
+	//~ Begin FTickableEditorObject interface
+	virtual void Tick(float DeltaTime) override;
+	virtual ETickableTickType GetTickableTickType() const override { return ETickableTickType::Conditional; }
+	virtual bool IsTickable() const override;
+	virtual TStatId GetStatId() const override;
+	//~ End FTickableEditorObject interface
+
+protected:
+	TSet<AGeneratedDynamicMeshActor*> ActiveGeneratedActors;
+};
diff --git a/Engine/Plugins/Experimental/GeometryScripting/Source/GeometryScriptingEditor/Public/GeometryActors/GeneratedDynamicMeshActor.h b/Engine/Plugins/Experimental/GeometryScripting/Source/GeometryScriptingEditor/Public/GeometryActors/GeneratedDynamicMeshActor.h
new file mode 100644
--- /dev/null
+++ b/Engine/Plugins/Experimental/GeometryScripting/Source/GeometryScriptingEditor/Public/GeometryActors/GeneratedDynamicMeshActor.h
@@ -0,0 +1,63 @@
+// Copyright Epic Games, Inc. All Rights Reserved.
+
+#pragma once
+
+#include "DynamicMeshActor.h"
+
+#include "GeneratedDynamicMeshActor.generated.h"
+
+
+/**
+ * AGeneratedDynamicMeshActor is an Editor-only subclass of ADynamicMeshActor that provides 
+ * special support for dynamic procedural generation of meshes in the Editor, eg via Blueprints. 
+ * Expensive procedural generation implemented via BP can potentially cause major problems in 
+ * the Editor, in particular with interactive performance. AGeneratedDynamicMeshActor provides
+ * special infrastructure for this use case. Essentially, instead of doing procedural generation
+ * in the Construction Script, a BP-implementable event OnRebuildGeneratedMesh is available,
+ * and doing the procedural mesh regeneration when that function fires will generally provide
+ * better in-Editor interactive performance.
+ */
+UCLASS()
+class GEOMETRYSCRIPTINGEDITOR_API AGeneratedDynamicMeshActor : public ADynamicMeshActor
+{
+	GENERATED_UCLASS_BODY()
+
+public:
+
+	virtual ~AGeneratedDynamicMeshActor();
+
+	/**
+	 * This event will be fired to notify the BP that the generated Mesh should
+	 * be rebuilt. GeneratedDynamicMeshActor BP subclasses should rebuild their 
+	 * meshes on this event, instead of doing so directly from the Construction Script.
+	 */
+	UFUNCTION(BlueprintImplementableEvent, CallInEditor, Category = "Events")
+	void OnRebuildGeneratedMesh();
+
+
+	/**
+	 * This function will fire the OnRebuildGeneratedMesh event if the actor has been
+	 * marked for a pending rebuild (eg via OnConstruction)
+	 */
+	virtual void ExecuteRebuildGeneratedMeshIfPending();
+
+
+public:
+
+	// AActor overrides
+	virtual void OnConstruction(const FTransform& Transform) override;
+	virtual void Destroyed() override;
+
+
+protected:
+	// this internal flag is set in OnConstruction, and will cause ExecuteRebuildGeneratedMesh to
+	// fire the OnRebuildGeneratedMesh event, after which the flag will be cleared
+	bool bGeneratedMeshRebuildPending = false;
+
+	// indicates that this Actor is registered with the UEditorGeometryGenerationSubsystem, which 
+	// is where the mesh rebuilds are executed
+	bool bIsRegisteredWithGenerationManager = false;
+
+	virtual void RegisterWithGenerationManager();
+	virtual void UnregisterWithGenerationManager();
+};
diff --git a/Engine/Source/Runtime/GeometryFramework/Private/DynamicMeshActor.cpp b/Engine/Source/Runtime/GeometryFramework/Private/DynamicMeshActor.cpp
--- a/Engine/Source/Runtime/GeometryFramework/Private/DynamicMeshActor.cpp
+++ b/Engine/Source/Runtime/GeometryFramework/Private/DynamicMeshActor.cpp
@@ -19,16 +19,6 @@ ADynamicMeshActor::ADynamicMeshActor(const FObjectInitializer& ObjectInitializer
 	DynamicMeshComponent->SetMaterial(0, UMaterial::GetDefaultMaterial(MD_Surface));		// is this necessary?
 
 	SetRootComponent(DynamicMeshComponent);
-
-	//
-	// Configure ADynamicMeshActor to always Tick(). This is necessary for the bIsEditorGeneratedMeshActor 
-	// Tick-in-Editor support below to work w/o any other user intervention. Generally always Ticking is bad behavior, 
-	// however DynamicMeshActor is inherently expensive and so we do not expect that large numbers of them
-	// will be used in a performance-sensitive context.
-	//
-	PrimaryActorTick.bCanEverTick = true;
-	PrimaryActorTick.bStartWithTickEnabled = true;
-	PrimaryActorTick.SetTickFunctionEnable(true);
 }
 
 
@@ -94,48 +84,4 @@ void ADynamicMeshActor::FreeAllComputeMeshes()
 
 
 
-
-
-void ADynamicMeshActor::Tick(float DeltaTime)
-{
-	Super::Tick(DeltaTime);
-
-#if WITH_EDITOR
-	if (bIsEditorGeneratedMeshActor && bGeneratedMeshRebuildPending)
-	{
-		OnEditorRebuildGeneratedMesh();
-		bGeneratedMeshRebuildPending = false;
-	}
-#endif
-}
-
-
-
-void ADynamicMeshActor::OnConstruction(const FTransform& Transform)
-{
-	Super::OnConstruction(Transform);
-
-#if WITH_EDITOR
-	if (bIsEditorGeneratedMeshActor)
-	{
-		bGeneratedMeshRebuildPending = true;
-	}
-#endif
-}
-
-
-#if WITH_EDITOR
-bool ADynamicMeshActor::ShouldTickIfViewportsOnly() const
-{
-	if (bIsEditorGeneratedMeshActor && GetWorld() != nullptr && GetWorld()->WorldType == EWorldType::Editor)
-	{
-		return true;
-	}
-	return false;
-}
-#endif
-
-
-
-
 #undef LOCTEXT_NAMESPACE
\ No newline at end of file
diff --git a/Engine/Source/Runtime/GeometryFramework/Public/DynamicMeshActor.h b/Engine/Source/Runtime/GeometryFramework/Public/DynamicMeshActor.h
--- a/Engine/Source/Runtime/GeometryFramework/Public/DynamicMeshActor.h
+++ b/Engine/Source/Runtime/GeometryFramework/Public/DynamicMeshActor.h
@@ -63,73 +63,6 @@ public:
 	/** Release all compute meshes that the Pool has allocated, and then release them from the Pool, so that they will be garbage-collected */
 	UFUNCTION(BlueprintCallable, Category = DynamicMeshActor)
 	void FreeAllComputeMeshes();
-
-
-
-
-	//
-	// In-Editor Generated/Procedural Mesh Actor support. 
-	// 
-	// These members are used to help with the creation of DynamicMeshActors that generate
-	// meshes dynamically in the Editor, for example via Blueprints. Expensive procedural 
-	// generation via BP can potentially cause major problems in the Editor, and so the
-	// implmentor of a BP Actor mesh generator needs to take some care.
-	// 
-	// If bIsEditorGeneratedMeshActor = true, then OnEditorRebuildGeneratedMesh
-	// will fire after the Construction script runs for this Actor. OnEditorRebuildGeneratedMesh does not directly
-	// fire from the Construction Script, but rather in the next Tick. This works around a fundamental
-	// limitation related to "mouse event priority" and slider/transform input in the UE Editor.
-	// Essentially doing even a moderately expensive mesh generation operation directly in the 
-	// Construction Script will appear to freeze the Viewport, while doing the same operation in the
-	// OnEditorRebuildGeneratedMesh implementation will remain responsive (albeit at a low FPS rate).
-	// 
-	// In future this may be extended to (eg) throttle procedural regeneration to further improve
-	// interactivity. 
-	// 
-	// Note that the current implementation requires in-Editor Ticking for this Actor. If bIsEditorGeneratedMeshActor
-	// is true, then ShouldTickIfViewportsOnly() will return true, which means that the Actor Tick Event
-	// will also fire in Editor viewports. This means that any game logic wired to the Tick Event
-	// in the Actor BP will aso run outside of PIE!
-	// 
-	// Note that the above behavior currently only occurs in the Editor, where the Construction script 
-	// is frequently re-executed (on Actor Transform and Property Editing). 
-	//
-public:
-
-	/**
-	 * Set this flag to true in a BP Subclass of DynamicMeshActor that procedurally generates it's
-	 * own mesh via BP (eg with Geometry Script). This will cause OnEditorRebuildGeneratedMesh to fire after 
-	 * the Construction Script is run or re-run. Regenerating the mesh on this event, rather than directly
-	 * in the Construction Script, will result in better interactive performance in the Editor.
-	 */
-	UPROPERTY(EditAnywhere, Category = DynamicMeshActor )
-	bool bIsEditorGeneratedMeshActor = false;
-
-	/**
-	 * This event will fire from Tick() to notify listeners that the internal Mesh should
-	 * be rebuilt. Procedural Mesh Generator Actors built in BP's derived from DynamicMeshActor
-	 * should listen for this event and use it to rebuild their meshes, instead of doing so
-	 * directly from the Construction Script.
-	 */
-	UFUNCTION(BlueprintImplementableEvent, CallInEditor, Category = "Events")
-	void OnEditorRebuildGeneratedMesh();
-
-
-
-	// the functions below are overridden to implement the above behavior
-public:
-	virtual void Tick(float DeltaTime) override;
-
-#if WITH_EDITOR
-	virtual bool ShouldTickIfViewportsOnly() const override;
-#endif
-
-	virtual void OnConstruction(const FTransform& Transform) override;
-
-protected:
-	// this internal flag is set in OnConstruction if bIsEditorGeneratedMeshActor=true, and will cause
-	// OnEditorRebuildGeneratedMesh to fire in Tick(), after which the flag is cleared
-	bool bGeneratedMeshRebuildPending = false;
 };
 
 
