GeometryFramework: Add Raytracing support for Base/DynamicMeshComponent.
- add FRayTracingGeometry structures to FMeshRenderBufferSet for primary and secondary triangle indices
- add UpdateRaytracingGeometryIfEnabled() function which rebuilds these raytracing data structures. This does a full rebuild, not clear if any kind of incremental update is possible currently, to revisit in future.
- call the above in the various places where we modify the existing index/vertex buffers.
- add UPrimitiveComponent::GetDynamicRayTracingInstances() implementation, mirrors GetDynamicMeshElements() implementation, uses new DrawRayTracingBatch() function

#rnx

diff --git a/Engine/Source/Runtime/GeometryFramework/GeometryFramework.Build.cs b/Engine/Source/Runtime/GeometryFramework/GeometryFramework.Build.cs
--- a/Engine/Source/Runtime/GeometryFramework/GeometryFramework.Build.cs
+++ b/Engine/Source/Runtime/GeometryFramework/GeometryFramework.Build.cs
@@ -6,9 +6,12 @@ public class GeometryFramework : ModuleRules
 {	
 	public GeometryFramework(ReadOnlyTargetRules Target) : base(Target)
 	{
-        //PCHUsage = ModuleRules.PCHUsageMode.UseExplicitOrSharedPCHs;
+		//PCHUsage = ModuleRules.PCHUsageMode.UseExplicitOrSharedPCHs;
 
-        PublicDependencyModuleNames.AddRange(
+		// this is required for raytracing suppport?
+		PublicIncludePaths.Add("../Shaders/Shared");
+
+		PublicDependencyModuleNames.AddRange(
 			new string[] {
 				"Core",
 				"CoreUObject",
diff --git a/Engine/Source/Runtime/GeometryFramework/Private/Components/BaseDynamicMeshComponent.cpp b/Engine/Source/Runtime/GeometryFramework/Private/Components/BaseDynamicMeshComponent.cpp
--- a/Engine/Source/Runtime/GeometryFramework/Private/Components/BaseDynamicMeshComponent.cpp
+++ b/Engine/Source/Runtime/GeometryFramework/Private/Components/BaseDynamicMeshComponent.cpp
@@ -12,6 +12,20 @@ UBaseDynamicMeshComponent::UBaseDynamicMeshComponent(const FObjectInitializer& O
 
 
 
+#if WITH_EDITOR
+void UBaseDynamicMeshComponent::PostEditChangeProperty(FPropertyChangedEvent& PropertyChangedEvent)
+{
+	Super::PostEditChangeProperty(PropertyChangedEvent);
+
+	const FName PropName = PropertyChangedEvent.GetPropertyName();
+	if ( (PropName == GET_MEMBER_NAME_CHECKED(UBaseDynamicMeshComponent, bEnableRaytracing))  )
+	{
+		OnRaytracingStateChanged();
+	}
+}
+#endif
+
+
 void UBaseDynamicMeshComponent::SetShadowsEnabled(bool bEnabled)
 {
 	// finish any drawing so that we can be certain our SceneProxy is no longer in use before we rebuild it below
@@ -87,6 +101,31 @@ bool UBaseDynamicMeshComponent::GetSecondaryBuffersVisibility() const
 }
 
 
+void UBaseDynamicMeshComponent::SetEnableRaytracing(bool bSetEnabled)
+{
+	if (bEnableRaytracing != bSetEnabled)
+	{
+		bEnableRaytracing = bSetEnabled;
+		OnRaytracingStateChanged();
+	}
+}
+
+bool UBaseDynamicMeshComponent::GetEnableRaytracing() const
+{
+	return bEnableRaytracing;
+}
+
+void UBaseDynamicMeshComponent::OnRaytracingStateChanged()
+{
+	// finish any drawing so that we can be certain our SceneProxy is no longer in use before we rebuild it below
+	FlushRenderingCommands();
+
+	// force immediate rebuild of the SceneProxy
+	if (IsRegistered())
+	{
+		ReregisterComponent();
+	}
+}
 
 
 int32 UBaseDynamicMeshComponent::GetNumMaterials() const
diff --git a/Engine/Source/Runtime/GeometryFramework/Public/Components/BaseDynamicMeshComponent.h b/Engine/Source/Runtime/GeometryFramework/Public/Components/BaseDynamicMeshComponent.h
--- a/Engine/Source/Runtime/GeometryFramework/Public/Components/BaseDynamicMeshComponent.h
+++ b/Engine/Source/Runtime/GeometryFramework/Public/Components/BaseDynamicMeshComponent.h
@@ -190,7 +190,7 @@ public:
 	 * if true, we always show the wireframe on top of the shaded mesh, even when not in wireframe mode
 	 * @todo: this should not be public, access via Set/Get once all usage is cleaned up
 	 */
-	UPROPERTY(EditAnywhere, BlueprintReadWrite, Category = "Dynamic Mesh Component", meta = (DisplayName = "Wireframe Overlay") )
+	UPROPERTY(EditAnywhere, BlueprintReadWrite, Category = "Dynamic Mesh Component|Rendering", meta = (DisplayName = "Wireframe Overlay") )
 	bool bExplicitShowWireframe = false;
 
 	/**
@@ -321,6 +321,33 @@ protected:
 
 
 
+	//===============================================================================================================
+	// Raytracing support. Must be enabled for various rendering effects. 
+	// However, note that in actual "dynamic" contexts (ie where the mesh is changing every frame),
+	// enabling Raytracing support has additional renderthread performance costs and does
+	// not currently support partial updates in the SceneProxy.
+public:
+
+	UPROPERTY(EditAnywhere, BlueprintReadWrite, Category = "Dynamic Mesh Component|Rendering")
+	bool bEnableRaytracing = false;
+
+	/**
+	 * Enable/Disable raytracing support. This is an expensive call as it flushes
+	 * the rendering queue and forces an immediate rebuild of the SceneProxy.
+	 */
+	UFUNCTION(BlueprintCallable, Category = "Dynamic Mesh Component")
+	virtual void SetEnableRaytracing(bool bSetEnabled);
+
+	/**
+	 * @return true if raytracing support is currently enabled
+	 */
+	UFUNCTION(BlueprintCallable, Category = "Dynamic Mesh Component")
+	virtual bool GetEnableRaytracing() const;
+
+protected:
+	virtual void OnRaytracingStateChanged();
+
+
 
 	//===============================================================================================================
 	// Standard Component interfaces
@@ -336,6 +363,11 @@ public:
 
 	virtual void SetNumMaterials(int32 NumMaterials);
 
+	//~ UObject Interface.
+#if WITH_EDITOR
+	void PostEditChangeProperty(FPropertyChangedEvent& PropertyChangedEvent) override;
+#endif
+
 	UPROPERTY()
 	TArray<TObjectPtr<UMaterialInterface>> BaseMaterials;
 };
diff --git a/Engine/Source/Runtime/GeometryFramework/Public/Components/BaseDynamicMeshSceneProxy.h b/Engine/Source/Runtime/GeometryFramework/Public/Components/BaseDynamicMeshSceneProxy.h
--- a/Engine/Source/Runtime/GeometryFramework/Public/Components/BaseDynamicMeshSceneProxy.h
+++ b/Engine/Source/Runtime/GeometryFramework/Public/Components/BaseDynamicMeshSceneProxy.h
@@ -13,6 +13,8 @@
 #include "DynamicMeshBuilder.h"
 #include "Components/BaseDynamicMeshComponent.h"
 #include "Materials/Material.h"
+#include "RayTracingDefinitions.h"
+#include "RayTracingInstance.h"
 
 using UE::Geometry::FDynamicMesh3;
 using UE::Geometry::FDynamicMeshAttributeSet;
@@ -62,6 +64,19 @@ public:
 	 */
 	FDynamicMeshIndexBuffer32 SecondaryIndexBuffer;
 
+	/**
+	 * configure whether raytracing should be enabled for this RenderBufferSet
+	 */
+	bool bEnableRaytracing = false;
+
+#if RHI_RAYTRACING
+	/**
+	 * Raytracing buffers
+	 */
+	FRayTracingGeometry PrimaryRayTracingGeometry;
+	FRayTracingGeometry SecondaryRayTracingGeometry;
+	bool bIsRayTracingDataValid = false;
+#endif
 
 	/**
 	 * In situations where we want to *update* the existing Vertex or Index buffers, we need to synchronize
@@ -96,6 +111,14 @@ public:
 			{
 				SecondaryIndexBuffer.ReleaseResource();
 			}
+
+#if RHI_RAYTRACING
+			if (bEnableRaytracing)
+			{
+				PrimaryRayTracingGeometry.ReleaseResource();
+				SecondaryRayTracingGeometry.ReleaseResource();
+			}
+#endif
 		}
 	}
 
@@ -140,11 +163,21 @@ public:
 		{
 			SecondaryIndexBuffer.InitResource();
 		}
+
+		InvalidateRayTracingData();
+		ValidateRayTracingData();		// currently we are immediately validating. This may be revisited in future.
 	}
 
 
+	/**
+	 * Fast path to only update the primary and secondary index buffers. This can be used
+	 * when (eg) the secondary index buffer is being used to highlight/hide a subset of triangles.
+	 * @warning This can only be called on the Rendering Thread.
+	 */
 	void UploadIndexBufferUpdate()
 	{
+		// todo: can this be done with RHI locking and memcpy, like in TransferVertexUpdateToGPU?
+
 		check(IsInRenderingThread());
 		if (IndexBuffer.Indices.Num() > 0)
 		{
@@ -154,11 +187,24 @@ public:
 		{
 			InitOrUpdateResource(&SecondaryIndexBuffer);
 		}
+
+		InvalidateRayTracingData();
+		ValidateRayTracingData();		// currently we are immediately validating. This may be revisited in future.
 	}
 
 
+	/**
+	 * Fast path to only update vertex buffers. This path rebuilds all the
+	 * resources and reconfigures the vertex factory, so the counts/etc could be modified.
+	 * @warning This can only be called on the Rendering Thread.
+	 */
 	void UploadVertexUpdate(bool bPositions, bool bMeshAttribs, bool bColors)
 	{
+		// todo: look at calls to this function, it seems possible that TransferVertexUpdateToGPU
+		// could be used instead (which should be somewhat more efficient?). It's not clear if there
+		// are any situations where we would change vertex buffer size w/o also updating the index
+		// buffers (in which case we are fully rebuilding the buffers...)
+
 		check(IsInRenderingThread());
 
 		if (TriangleCount == 0)
@@ -183,15 +229,21 @@ public:
 		this->PositionVertexBuffer.BindPositionVertexBuffer(&this->VertexFactory, Data);
 		this->StaticMeshVertexBuffer.BindTangentVertexBuffer(&this->VertexFactory, Data);
 		this->StaticMeshVertexBuffer.BindPackedTexCoordVertexBuffer(&this->VertexFactory, Data);
-		// currently no lightmaps support
-		//this->StaticMeshVertexBuffer.BindLightMapVertexBuffer(&this->VertexFactory, Data, LightMapIndex);
 		this->ColorVertexBuffer.BindColorVertexBuffer(&this->VertexFactory, Data);
 		this->VertexFactory.SetData(Data);
 
 		InitOrUpdateResource(&this->VertexFactory);
+
+		InvalidateRayTracingData();
+		ValidateRayTracingData();		// currently we are immediately validating. This may be revisited in future.
 	}
 
 
+	/**
+	 * Fast path to update various vertex buffers. This path does not support changing the
+	 * size/counts of any of the sub-buffers, a direct memcopy from the CPU-side buffer to the RHI buffer is used.
+	 * @warning This can only be called on the Rendering Thread.
+	 */
 	void TransferVertexUpdateToGPU(bool bPositions, bool bNormals, bool bTexCoords, bool bColors)
 	{
 		check(IsInRenderingThread());
@@ -202,42 +254,107 @@ public:
 
 		if (bPositions)
 		{
-			auto& VertexBuffer = this->PositionVertexBuffer;
+			FPositionVertexBuffer& VertexBuffer = this->PositionVertexBuffer;
 			void* VertexBufferData = RHILockBuffer(VertexBuffer.VertexBufferRHI, 0, VertexBuffer.GetNumVertices() * VertexBuffer.GetStride(), RLM_WriteOnly);
 			FMemory::Memcpy(VertexBufferData, VertexBuffer.GetVertexData(), VertexBuffer.GetNumVertices() * VertexBuffer.GetStride());
 			RHIUnlockBuffer(VertexBuffer.VertexBufferRHI);
 		}
 		if (bNormals)
 		{
-			auto& VertexBuffer = this->StaticMeshVertexBuffer;
+			FStaticMeshVertexBuffer& VertexBuffer = this->StaticMeshVertexBuffer;
 			void* VertexBufferData = RHILockBuffer(VertexBuffer.TangentsVertexBuffer.VertexBufferRHI, 0, VertexBuffer.GetTangentSize(), RLM_WriteOnly);
 			FMemory::Memcpy(VertexBufferData, VertexBuffer.GetTangentData(), VertexBuffer.GetTangentSize());
 			RHIUnlockBuffer(VertexBuffer.TangentsVertexBuffer.VertexBufferRHI);
 		}
 		if (bColors)
 		{
-			auto& VertexBuffer = this->ColorVertexBuffer;
+			FColorVertexBuffer& VertexBuffer = this->ColorVertexBuffer;
 			void* VertexBufferData = RHILockBuffer(VertexBuffer.VertexBufferRHI, 0, VertexBuffer.GetNumVertices() * VertexBuffer.GetStride(), RLM_WriteOnly);
 			FMemory::Memcpy(VertexBufferData, VertexBuffer.GetVertexData(), VertexBuffer.GetNumVertices() * VertexBuffer.GetStride());
 			RHIUnlockBuffer(VertexBuffer.VertexBufferRHI);
 		}
 		if (bTexCoords)
 		{
-			auto& VertexBuffer = this->StaticMeshVertexBuffer;
+			FStaticMeshVertexBuffer& VertexBuffer = this->StaticMeshVertexBuffer;
 			void* VertexBufferData = RHILockBuffer(VertexBuffer.TexCoordVertexBuffer.VertexBufferRHI, 0, VertexBuffer.GetTexCoordSize(), RLM_WriteOnly);
 			FMemory::Memcpy(VertexBufferData, VertexBuffer.GetTexCoordData(), VertexBuffer.GetTexCoordSize());
 			RHIUnlockBuffer(VertexBuffer.TexCoordVertexBuffer.VertexBufferRHI);
 		}
+
+		InvalidateRayTracingData();
+		ValidateRayTracingData();		// currently we are immediately validating. This may be revisited in future.
+	}
+
+
+	void InvalidateRayTracingData()
+	{
+#if RHI_RAYTRACING
+		bIsRayTracingDataValid = false;
+#endif
+	}
+
+	// Verify that valid raytracing data is available. This will cause a rebuild of the
+	// raytracing data if any of our buffers have been modified. Currently this is called
+	// by GetDynamicRayTracingInstances to ensure the RT data is available when needed.
+	void ValidateRayTracingData()
+	{
+#if RHI_RAYTRACING
+		if (bIsRayTracingDataValid == false && IsRayTracingEnabled() && bEnableRaytracing)
+		{
+			UpdateRaytracingGeometryIfEnabled();
+
+			bIsRayTracingDataValid = true;
+		}
+#endif
 	}
 
 
+protected:
+
+	// rebuild raytracing data for current buffers
+	void UpdateRaytracingGeometryIfEnabled()
+	{
+#if RHI_RAYTRACING
+		// do we always want to do this?
+		PrimaryRayTracingGeometry.ReleaseResource();		
+		SecondaryRayTracingGeometry.ReleaseResource();
+			
+		for (int32 k = 0; k < 2; ++k)
+		{
+			FDynamicMeshIndexBuffer32& UseIndexBuffer = (k == 0) ? IndexBuffer : SecondaryIndexBuffer;
+			if (UseIndexBuffer.Indices.Num() == 0)
+			{
+				continue;
+			}
+
+			FRayTracingGeometry& RayTracingGeometry = (k == 0) ? PrimaryRayTracingGeometry : SecondaryRayTracingGeometry;
+
+			FRayTracingGeometryInitializer Initializer;
+			Initializer.IndexBuffer = UseIndexBuffer.IndexBufferRHI;
+			Initializer.TotalPrimitiveCount = UseIndexBuffer.Indices.Num() / 3;
+			Initializer.GeometryType = RTGT_Triangles;
+			Initializer.bFastBuild = true;
+			Initializer.bAllowUpdate = false;
+
+			RayTracingGeometry.SetInitializer(Initializer);
+			RayTracingGeometry.InitResource();
+
+			FRayTracingGeometrySegment Segment;
+			Segment.VertexBuffer = PositionVertexBuffer.VertexBufferRHI;
+			Segment.NumPrimitives = RayTracingGeometry.Initializer.TotalPrimitiveCount;
+			Segment.MaxVertices = PositionVertexBuffer.GetNumVertices();
+			RayTracingGeometry.Initializer.Segments.Add(Segment);
 
+			RayTracingGeometry.UpdateRHI();
+		}
+#endif
+	}
 
 
 
 	/**
 	 * Initializes a render resource, or update it if already initialized.
-	 * @wearning This function can only be called on the Render Thread
+	 * @warning This function can only be called on the Render Thread
 	 */
 	void InitOrUpdateResource(FRenderResource* Resource)
 	{
@@ -334,10 +451,14 @@ protected:
 	// use to control access to AllocatedBufferSets 
 	FCriticalSection AllocatedSetsLock;
 
+	// control raytracing support
+	bool bEnableRaytracing = false;
+
 public:
 	FBaseDynamicMeshSceneProxy(UBaseDynamicMeshComponent* Component)
 		: FPrimitiveSceneProxy(Component),
-		ParentBaseComponent(Component)
+		ParentBaseComponent(Component),
+		bEnableRaytracing(Component->GetEnableRaytracing())
 	{
 	}
 
@@ -383,6 +504,7 @@ public:
 		FMeshRenderBufferSet* RenderBufferSet = new FMeshRenderBufferSet(GetScene().GetFeatureLevel());
 
 		RenderBufferSet->Material = UMaterial::GetDefaultMaterial(MD_Surface);
+		RenderBufferSet->bEnableRaytracing = this->bEnableRaytracing;
 
 		AllocatedSetsLock.Lock();
 		AllocatedBufferSets.Add(RenderBufferSet);
@@ -805,7 +927,7 @@ public:
 		uint32 VisibilityMap, 
 		FMeshElementCollector& Collector) const override
 	{
-		QUICK_SCOPE_CYCLE_COUNTER(STAT_OctreeDynamicMeshSceneProxy_GetDynamicMeshElements);
+		QUICK_SCOPE_CYCLE_COUNTER(STAT_BaseDynamicMeshSceneProxy_GetDynamicMeshElements);
 
 		const bool bWireframe = (AllowDebugViewmodes() && ViewFamily.EngineShowFlags.Wireframe)
 			|| ParentBaseComponent->GetEnableWireframeRenderPass();
@@ -932,6 +1054,128 @@ public:
 
 
 
+#if RHI_RAYTRACING
+
+	virtual bool IsRayTracingRelevant() const override { return true; }
+
+	virtual void GetDynamicRayTracingInstances(FRayTracingMaterialGatheringContext& Context, TArray<FRayTracingInstance>& OutRayTracingInstances) override
+	{
+		QUICK_SCOPE_CYCLE_COUNTER(STAT_BaseDynamicMeshSceneProxy_GetDynamicRayTracingInstances);
+
+		ESceneDepthPriorityGroup DepthPriority = SDPG_World;
+
+		TArray<FMeshRenderBufferSet*> Buffers;
+		GetActiveRenderBufferSets(Buffers);
+
+		FMaterialRenderProxy* SecondaryMaterialProxy =
+			ParentBaseComponent->HasSecondaryRenderMaterial() ? ParentBaseComponent->GetSecondaryRenderMaterial()->GetRenderProxy() : nullptr;
+		bool bDrawSecondaryBuffers = ParentBaseComponent->GetSecondaryBuffersVisibility();
+
+		bool bHasPrecomputedVolumetricLightmap;
+		FMatrix PreviousLocalToWorld;
+		int32 SingleCaptureIndex;
+		bool bOutputVelocity;
+		GetScene().GetPrimitiveUniformShaderParameters_RenderThread(GetPrimitiveSceneInfo(), bHasPrecomputedVolumetricLightmap, PreviousLocalToWorld, SingleCaptureIndex, bOutputVelocity);
+
+		// is it safe to share this between primary and secondary raytracing batches?
+		FDynamicPrimitiveUniformBuffer& DynamicPrimitiveUniformBuffer = Context.RayTracingMeshResourceCollector.AllocateOneFrameResource<FDynamicPrimitiveUniformBuffer>();
+		DynamicPrimitiveUniformBuffer.Set(GetLocalToWorld(), PreviousLocalToWorld, GetBounds(), GetLocalBounds(), true, bHasPrecomputedVolumetricLightmap, DrawsVelocity(), bOutputVelocity);
+
+		// Draw the active buffer sets
+		for (FMeshRenderBufferSet* BufferSet : Buffers)
+		{
+			UMaterialInterface* UseMaterial = BufferSet->Material;
+			if (ParentBaseComponent->HasOverrideRenderMaterial(0))
+			{
+				UseMaterial = ParentBaseComponent->GetOverrideRenderMaterial(0);
+			}
+			FMaterialRenderProxy* MaterialProxy = UseMaterial->GetRenderProxy();
+
+			if (BufferSet->TriangleCount == 0)
+			{
+				continue;
+			}
+			if (ensure(BufferSet->bIsRayTracingDataValid) == false)
+			{
+				continue;
+			}
+
+			// Lock buffers so that they aren't modified while we are submitting them.
+			FScopeLock BuffersLock(&BufferSet->BuffersLock);
+
+			// draw primary index buffer
+			if (BufferSet->IndexBuffer.Indices.Num() > 0 
+				&& BufferSet->PrimaryRayTracingGeometry.RayTracingGeometryRHI.IsValid())
+			{
+				ensure(BufferSet->PrimaryRayTracingGeometry.Initializer.IndexBuffer.IsValid());
+				DrawRayTracingBatch(Context, *BufferSet, BufferSet->IndexBuffer, BufferSet->PrimaryRayTracingGeometry, MaterialProxy, DepthPriority, DynamicPrimitiveUniformBuffer, OutRayTracingInstances);
+			}
+
+			// draw secondary index buffer if we have it, falling back to base material if we don't have the Secondary material
+			FMaterialRenderProxy* UseSecondaryMaterialProxy = (SecondaryMaterialProxy != nullptr) ? SecondaryMaterialProxy : MaterialProxy;
+			if (bDrawSecondaryBuffers 
+				&& BufferSet->SecondaryIndexBuffer.Indices.Num() > 0 
+				&& UseSecondaryMaterialProxy != nullptr
+				&& BufferSet->SecondaryRayTracingGeometry.RayTracingGeometryRHI.IsValid() )
+			{
+				ensure(BufferSet->SecondaryRayTracingGeometry.Initializer.IndexBuffer.IsValid());
+				DrawRayTracingBatch(Context, *BufferSet, BufferSet->SecondaryIndexBuffer, BufferSet->SecondaryRayTracingGeometry, UseSecondaryMaterialProxy, DepthPriority, DynamicPrimitiveUniformBuffer, OutRayTracingInstances);
+			}
+		}
+	}
+
+
+
+	/**
+	* Draw a single-frame raytracing FMeshBatch for a FMeshRenderBufferSet
+	*/
+	virtual void DrawRayTracingBatch(
+		FRayTracingMaterialGatheringContext& Context,
+		const FMeshRenderBufferSet& RenderBuffers,
+		const FDynamicMeshIndexBuffer32& IndexBuffer,
+		FRayTracingGeometry& RayTracingGeometry,
+		FMaterialRenderProxy* UseMaterialProxy,
+		ESceneDepthPriorityGroup DepthPriority,
+		FDynamicPrimitiveUniformBuffer& DynamicPrimitiveUniformBuffer,
+		TArray<FRayTracingInstance>& OutRayTracingInstances	) const
+	{
+		ensure(RayTracingGeometry.Initializer.IndexBuffer.IsValid());
+
+		FRayTracingInstance RayTracingInstance;
+		RayTracingInstance.Geometry = &RayTracingGeometry;
+		RayTracingInstance.InstanceTransforms.Add(GetLocalToWorld());
+
+		uint32 SectionIdx = 0;
+		FMeshBatch MeshBatch;
+
+		MeshBatch.VertexFactory = &RenderBuffers.VertexFactory;
+		MeshBatch.SegmentIndex = 0;
+		MeshBatch.MaterialRenderProxy = UseMaterialProxy;
+		MeshBatch.ReverseCulling = IsLocalToWorldDeterminantNegative();
+		MeshBatch.Type = PT_TriangleList;
+		MeshBatch.DepthPriorityGroup = DepthPriority;
+		MeshBatch.bCanApplyViewModeOverrides = false;
+		MeshBatch.CastRayTracedShadow = IsShadowCast(Context.ReferenceView);
+
+		FMeshBatchElement& BatchElement = MeshBatch.Elements[0];
+		BatchElement.IndexBuffer = &IndexBuffer;
+		BatchElement.PrimitiveUniformBufferResource = &DynamicPrimitiveUniformBuffer.UniformBuffer;
+		BatchElement.FirstIndex = 0;
+		BatchElement.NumPrimitives = IndexBuffer.Indices.Num() / 3;
+		BatchElement.MinVertexIndex = 0;
+		BatchElement.MaxVertexIndex = RenderBuffers.PositionVertexBuffer.GetNumVertices() - 1;
+
+		RayTracingInstance.Materials.Add(MeshBatch);
+
+		RayTracingInstance.BuildInstanceMaskAndFlags(GetScene().GetFeatureLevel());
+		OutRayTracingInstances.Add(RayTracingInstance);
+	}
+
+
+#endif // RHI_RAYTRACING
+
+
+
 
 
 
