GeometryFramework: remove ensure that will always be hit before DynamicMeshComponent RT geomety is initialized
#rnx

diff --git a/Engine/Source/Runtime/GeometryFramework/Public/Components/BaseDynamicMeshSceneProxy.h b/Engine/Source/Runtime/GeometryFramework/Public/Components/BaseDynamicMeshSceneProxy.h
--- a/Engine/Source/Runtime/GeometryFramework/Public/Components/BaseDynamicMeshSceneProxy.h
+++ b/Engine/Source/Runtime/GeometryFramework/Public/Components/BaseDynamicMeshSceneProxy.h
@@ -1095,7 +1095,7 @@ public:
 			{
 				continue;
 			}
-			if (ensure(BufferSet->bIsRayTracingDataValid) == false)
+			if (BufferSet->bIsRayTracingDataValid == false)
 			{
 				continue;
 			}
