GeometryScript: automatically clear mesh before rebuild in GeneratedDynamicMeshActor. Expose ADynamicMeshActor members. Bump up size of "empty" mesh bounding box in DynamicMeshComponent to prevent log spam.
#rnx

diff --git a/Engine/Plugins/Experimental/GeometryScripting/Source/GeometryScriptingEditor/Private/GeometryActors/GeneratedDynamicMeshActor.cpp b/Engine/Plugins/Experimental/GeometryScripting/Source/GeometryScriptingEditor/Private/GeometryActors/GeneratedDynamicMeshActor.cpp
--- a/Engine/Plugins/Experimental/GeometryScripting/Source/GeometryScriptingEditor/Private/GeometryActors/GeneratedDynamicMeshActor.cpp
+++ b/Engine/Plugins/Experimental/GeometryScripting/Source/GeometryScriptingEditor/Private/GeometryActors/GeneratedDynamicMeshActor.cpp
@@ -68,6 +68,11 @@ void AGeneratedDynamicMeshActor::ExecuteRebuildGeneratedMeshIfPending()
 {
 	if (bGeneratedMeshRebuildPending)
 	{
+		if (bResetOnRebuild && DynamicMeshComponent && DynamicMeshComponent->GetDynamicMesh())
+		{
+			DynamicMeshComponent->GetDynamicMesh()->Reset();
+		}
+
 		FEditorScriptExecutionGuard Guard;
 		OnRebuildGeneratedMesh();
 		bGeneratedMeshRebuildPending = false;
diff --git a/Engine/Plugins/Experimental/GeometryScripting/Source/GeometryScriptingEditor/Public/GeometryActors/GeneratedDynamicMeshActor.h b/Engine/Plugins/Experimental/GeometryScripting/Source/GeometryScriptingEditor/Public/GeometryActors/GeneratedDynamicMeshActor.h
--- a/Engine/Plugins/Experimental/GeometryScripting/Source/GeometryScriptingEditor/Public/GeometryActors/GeneratedDynamicMeshActor.h
+++ b/Engine/Plugins/Experimental/GeometryScripting/Source/GeometryScriptingEditor/Public/GeometryActors/GeneratedDynamicMeshActor.h
@@ -26,6 +26,13 @@ public:
 
 	virtual ~AGeneratedDynamicMeshActor();
 
+
+public:
+	/** If true, the internal UDynamicMesh will be cleared before the OnRebuildGeneratedMesh event is fired. */
+	UPROPERTY(Category = DynamicMeshActor, EditAnywhere, BlueprintReadWrite)
+	bool bResetOnRebuild = true;
+
+
 	/**
 	 * This event will be fired to notify the BP that the generated Mesh should
 	 * be rebuilt. GeneratedDynamicMeshActor BP subclasses should rebuild their 
diff --git a/Engine/Source/Runtime/GeometryFramework/Private/Components/DynamicMeshComponent.cpp b/Engine/Source/Runtime/GeometryFramework/Private/Components/DynamicMeshComponent.cpp
--- a/Engine/Source/Runtime/GeometryFramework/Private/Components/DynamicMeshComponent.cpp
+++ b/Engine/Source/Runtime/GeometryFramework/Private/Components/DynamicMeshComponent.cpp
@@ -280,7 +280,10 @@ void UDynamicMeshComponent::UpdateLocalBounds()
 	LocalBounds = GetMesh()->GetBounds(true);
 	if (LocalBounds.MaxDim() <= 0)
 	{
-		LocalBounds = FAxisAlignedBox3d(FVector3d::Zero(), FMathf::ZeroTolerance);
+		// If bbox is empty, set a very small bbox to avoid log spam/etc in other engine systems.
+		// The check used is generally IsNearlyZero(), which defaults to KINDA_SMALL_NUMBER, so set 
+		// a slightly larger box here to be above that threshold
+		LocalBounds = FAxisAlignedBox3d(FVector3d::Zero(), (double)(KINDA_SMALL_NUMBER + SMALL_NUMBER) );
 	}
 }
 
diff --git a/Engine/Source/Runtime/GeometryFramework/Public/DynamicMeshActor.h b/Engine/Source/Runtime/GeometryFramework/Public/DynamicMeshActor.h
--- a/Engine/Source/Runtime/GeometryFramework/Public/DynamicMeshActor.h
+++ b/Engine/Source/Runtime/GeometryFramework/Public/DynamicMeshActor.h
@@ -17,7 +17,7 @@ class GEOMETRYFRAMEWORK_API ADynamicMeshActor : public AActor
 {
 	GENERATED_UCLASS_BODY()
 
-private:
+protected:
 	UPROPERTY(Category = DynamicMeshActor, VisibleAnywhere, BlueprintReadOnly, meta = (ExposeFunctionCategories = "Mesh,Rendering,Physics,Components|StaticMesh", AllowPrivateAccess = "true"))
 	TObjectPtr<class UDynamicMeshComponent> DynamicMeshComponent;
 
@@ -36,9 +36,9 @@ public:
 
 public:
 	/** Control whether the DynamicMeshPool will be created when requested via GetComputeMeshPool() */
-	UPROPERTY(Category = DynamicMeshActor, VisibleAnywhere, BlueprintReadWrite)
+	UPROPERTY(Category = DynamicMeshActor, EditAnywhere, BlueprintReadWrite)
 	bool bEnableComputeMeshPool = true;
-private:
+protected:
 	/** The internal Mesh Pool, for use in DynamicMeshActor BPs. Use GetComputeMeshPool() to access this, as it will only be created on-demand if bEnableComputeMeshPool = true */
 	UPROPERTY(Transient)
 	TObjectPtr<UDynamicMeshPool> DynamicMeshPool;
