GeometryFramework:
- add DynamicMeshComponent::SetComplexAsSimpleCollisionEnabled(), SetDeferredCollisionUpdatesEnabled() to handle some common component variable configuration in BP.
- expose DynamicMeshComponent::UpdateCollision()
GeometryScripting:
- pass UDynamicMesh TargetMesh as parameter to GeneratedDynamicMeshActor::OnRebuildGeneratedMesh(), to simplify BPs and leave open the possibility of updating a compute mesh separate from the actual Component mesh
- Automatically enable (then disable) deferred collision updates on DynamicMeshComponent in AGeneratedDynamicMeshActor::ExecuteRebuildGeneratedMeshIfPending, otherwise every GeometryScript node will rebuild collision (!)
#rnx

diff --git a/Engine/Plugins/Experimental/GeometryScripting/Source/GeometryScriptingEditor/Private/GeometryActors/GeneratedDynamicMeshActor.cpp b/Engine/Plugins/Experimental/GeometryScripting/Source/GeometryScriptingEditor/Private/GeometryActors/GeneratedDynamicMeshActor.cpp
--- a/Engine/Plugins/Experimental/GeometryScripting/Source/GeometryScriptingEditor/Private/GeometryActors/GeneratedDynamicMeshActor.cpp
+++ b/Engine/Plugins/Experimental/GeometryScripting/Source/GeometryScriptingEditor/Private/GeometryActors/GeneratedDynamicMeshActor.cpp
@@ -68,14 +68,28 @@ void AGeneratedDynamicMeshActor::ExecuteRebuildGeneratedMeshIfPending()
 {
 	if (bGeneratedMeshRebuildPending)
 	{
+		// Automatically defer collision updates during generated mesh rebuild. If we do not do this, then
+		// each mesh change will result in collision being rebuilt, which is very expensive !!
+		bool bEnabledDeferredCollision = false;
+		if (DynamicMeshComponent->bDeferCollisionUpdates == false)
+		{
+			DynamicMeshComponent->SetDeferredCollisionUpdatesEnabled(true, false);
+			bEnabledDeferredCollision = true;
+		}
+
 		if (bResetOnRebuild && DynamicMeshComponent && DynamicMeshComponent->GetDynamicMesh())
 		{
 			DynamicMeshComponent->GetDynamicMesh()->Reset();
 		}
 
 		FEditorScriptExecutionGuard Guard;
-		OnRebuildGeneratedMesh();
+		OnRebuildGeneratedMesh(DynamicMeshComponent->GetDynamicMesh());
 		bGeneratedMeshRebuildPending = false;
+
+		if (bEnabledDeferredCollision)
+		{
+			DynamicMeshComponent->SetDeferredCollisionUpdatesEnabled(false, true);
+		}
 	}
 }
 
diff --git a/Engine/Plugins/Experimental/GeometryScripting/Source/GeometryScriptingEditor/Public/GeometryActors/GeneratedDynamicMeshActor.h b/Engine/Plugins/Experimental/GeometryScripting/Source/GeometryScriptingEditor/Public/GeometryActors/GeneratedDynamicMeshActor.h
--- a/Engine/Plugins/Experimental/GeometryScripting/Source/GeometryScriptingEditor/Public/GeometryActors/GeneratedDynamicMeshActor.h
+++ b/Engine/Plugins/Experimental/GeometryScripting/Source/GeometryScriptingEditor/Public/GeometryActors/GeneratedDynamicMeshActor.h
@@ -39,7 +39,7 @@ public:
 	 * meshes on this event, instead of doing so directly from the Construction Script.
 	 */
 	UFUNCTION(BlueprintImplementableEvent, CallInEditor, Category = "Events")
-	void OnRebuildGeneratedMesh();
+	void OnRebuildGeneratedMesh(UDynamicMesh* TargetMesh);
 
 
 	/**
diff --git a/Engine/Source/Runtime/GeometryFramework/Private/Components/DynamicMeshComponent.cpp b/Engine/Source/Runtime/GeometryFramework/Private/Components/DynamicMeshComponent.cpp
--- a/Engine/Source/Runtime/GeometryFramework/Private/Components/DynamicMeshComponent.cpp
+++ b/Engine/Source/Runtime/GeometryFramework/Private/Components/DynamicMeshComponent.cpp
@@ -1196,7 +1196,57 @@ void UDynamicMeshComponent::UpdateCollision(bool bOnlyIfPending)
 
 void UDynamicMeshComponent::EnableComplexAsSimpleCollision()
 {
-	bEnableComplexCollision = true;
-	CollisionType = ECollisionTraceFlag::CTF_UseComplexAsSimple;
-	UpdateCollision(true);
-}
\ No newline at end of file
+	SetComplexAsSimpleCollisionEnabled(true, true);
+}
+
+void UDynamicMeshComponent::SetComplexAsSimpleCollisionEnabled(bool bEnabled, bool bImmediateUpdate)
+{
+	bool bModified = false;
+	if (bEnabled)
+	{
+		if (bEnableComplexCollision == false)
+		{
+			bEnableComplexCollision = true;
+			bModified = true;
+		}
+		if (CollisionType != ECollisionTraceFlag::CTF_UseComplexAsSimple)
+		{
+			CollisionType = ECollisionTraceFlag::CTF_UseComplexAsSimple;
+			bModified = true;
+		}
+	}
+	else
+	{
+		if (bEnableComplexCollision == true)
+		{
+			bEnableComplexCollision = false;
+			bModified = true;
+		}
+		if (CollisionType == ECollisionTraceFlag::CTF_UseComplexAsSimple)
+		{
+			CollisionType = ECollisionTraceFlag::CTF_UseDefault;
+			bModified = true;
+		}
+	}
+	if (bModified)
+	{
+		InvalidatePhysicsData();
+	}
+	if (bImmediateUpdate)
+	{
+		UpdateCollision(true);
+	}
+}
+
+
+void UDynamicMeshComponent::SetDeferredCollisionUpdatesEnabled(bool bEnabled, bool bImmediateUpdate)
+{
+	if (bDeferCollisionUpdates != bEnabled)
+	{
+		bDeferCollisionUpdates = bEnabled;
+		if (bEnabled == false && bImmediateUpdate)
+		{
+			UpdateCollision(true);
+		}
+	}
+}
diff --git a/Engine/Source/Runtime/GeometryFramework/Public/Components/DynamicMeshComponent.h b/Engine/Source/Runtime/GeometryFramework/Public/Components/DynamicMeshComponent.h
--- a/Engine/Source/Runtime/GeometryFramework/Public/Components/DynamicMeshComponent.h
+++ b/Engine/Source/Runtime/GeometryFramework/Public/Components/DynamicMeshComponent.h
@@ -456,9 +456,28 @@ protected:
 	// Physics APIs
 	//
 public:
+	/**
+	 * calls SetComplexAsSimpleCollisionEnabled(true, true)
+	 */
 	UFUNCTION(BlueprintCallable, Category = "Dynamic Mesh Component")
 	void EnableComplexAsSimpleCollision();
 
+	/**
+	 * If bEnabled=true, sets bEnableComplexCollision=true and CollisionType=CTF_UseComplexAsSimple
+	 * If bEnabled=true, sets bEnableComplexCollision=false and CollisionType=CTF_UseDefault
+	 * @param bImmediateUpdate if true, UpdateCollision(true) is called
+	 */
+	UFUNCTION(BlueprintCallable, Category = "Dynamic Mesh Component")
+	void SetComplexAsSimpleCollisionEnabled(bool bEnabled, bool bImmediateUpdate = true);
+
+	/**
+	 * Set value of bDeferCollisionUpdates, when enabled, collision is not automatically recomputed each time the mesh changes.
+	 * @param bImmediateUpdate if true, UpdateCollision(true) is called if bEnabled=false, ie to force a collision rebuild
+	 */
+	UFUNCTION(BlueprintCallable, Category = "Dynamic Mesh Component")
+	void SetDeferredCollisionUpdatesEnabled(bool bEnabled, bool bImmediateUpdate = true);
+
+
 	virtual bool GetPhysicsTriMeshData(struct FTriMeshCollisionData* CollisionData, bool InUseAllTriData) override;
 	virtual bool ContainsPhysicsTriMeshData(bool InUseAllTriData) const override;
 	virtual bool WantsNegXTriMesh() override;
@@ -470,7 +489,8 @@ public:
 	 * Force an update of the Collision/Physics data for this Component.
 	 * @param bOnlyIfPending only update if a collision update is pending, ie the underlying DynamicMesh changed and bDeferCollisionUpdates is enabled
 	 */
-	virtual void UpdateCollision(bool bOnlyIfPending);
+	UFUNCTION(BlueprintCallable, Category = "Dynamic Mesh Component")
+	virtual void UpdateCollision(bool bOnlyIfPending = true);
 
 	/** Type of Collision Geometry to use for this Mesh */
 	UPROPERTY(EditAnywhere, BlueprintReadWrite, Category = "Dynamic Mesh Component|Collision")
