GeometryFramework: expose View Mode Overrides flag as a setting on BaseDynamicMeshComponent. This flag controls how the mesh is drawn in different Editor View Modes, eg in Wireframe mode, if view mode overrides are disabled, the mesh will be drawn to the Z buffer to remove hidden lines. It also effects the various lighting-only modes. This change defaults the overrides to enabled, ie changing the previous behavior which disabled them, as otherwise DynamicMeshActor meshes are opaque in the standard Editor Wireframe mode.

#rnx

diff --git a/Engine/Source/Runtime/GeometryFramework/Private/Components/BaseDynamicMeshComponent.cpp b/Engine/Source/Runtime/GeometryFramework/Private/Components/BaseDynamicMeshComponent.cpp
--- a/Engine/Source/Runtime/GeometryFramework/Private/Components/BaseDynamicMeshComponent.cpp
+++ b/Engine/Source/Runtime/GeometryFramework/Private/Components/BaseDynamicMeshComponent.cpp
@@ -20,7 +20,11 @@ void UBaseDynamicMeshComponent::PostEditChangeProperty(FPropertyChangedEvent& Pr
 	const FName PropName = PropertyChangedEvent.GetPropertyName();
 	if ( (PropName == GET_MEMBER_NAME_CHECKED(UBaseDynamicMeshComponent, bEnableRaytracing))  )
 	{
-		OnRaytracingStateChanged();
+		OnRenderingStateChanged(true);
+	}
+	else if ( (PropName == GET_MEMBER_NAME_CHECKED(UBaseDynamicMeshComponent, bEnableViewModeOverrides))  )
+	{
+		OnRenderingStateChanged(false);
 	}
 }
 #endif
@@ -28,22 +32,18 @@ void UBaseDynamicMeshComponent::PostEditChangeProperty(FPropertyChangedEvent& Pr
 
 void UBaseDynamicMeshComponent::SetShadowsEnabled(bool bEnabled)
 {
-	// finish any drawing so that we can be certain our SceneProxy is no longer in use before we rebuild it below
 	FlushRenderingCommands();
-
 	SetCastShadow(bEnabled);
-	//bCastDynamicShadow = bEnabled;
-
-	// apparently SceneProxy has to be fully rebuilt to change shadow state
+	OnRenderingStateChanged(true);
+}
 
-	// this marks the SceneProxy for rebuild, but not immediately, and possibly allows bad things to happen
-	// before the end of the frame
-	//MarkRenderStateDirty();
 
-	// force immediate rebuild of the SceneProxy
-	if (IsRegistered())
+void UBaseDynamicMeshComponent::SetViewModeOverridesEnabled(bool bEnabled)
+{
+	if (bEnableViewModeOverrides != bEnabled)
 	{
-		ReregisterComponent();
+		bEnableViewModeOverrides = bEnabled;
+		OnRenderingStateChanged(false);
 	}
 }
 
@@ -106,7 +106,7 @@ void UBaseDynamicMeshComponent::SetEnableRaytracing(bool bSetEnabled)
 	if (bEnableRaytracing != bSetEnabled)
 	{
 		bEnableRaytracing = bSetEnabled;
-		OnRaytracingStateChanged();
+		OnRenderingStateChanged(true);
 	}
 }
 
@@ -115,15 +115,22 @@ bool UBaseDynamicMeshComponent::GetEnableRaytracing() const
 	return bEnableRaytracing;
 }
 
-void UBaseDynamicMeshComponent::OnRaytracingStateChanged()
+void UBaseDynamicMeshComponent::OnRenderingStateChanged(bool bForceImmedateRebuild)
 {
-	// finish any drawing so that we can be certain our SceneProxy is no longer in use before we rebuild it below
-	FlushRenderingCommands();
+	if (bForceImmedateRebuild)
+	{
+		// finish any drawing so that we can be certain our SceneProxy is no longer in use before we rebuild it below
+		FlushRenderingCommands();
 
-	// force immediate rebuild of the SceneProxy
-	if (IsRegistered())
+		// force immediate rebuild of the SceneProxy
+		if (IsRegistered())
+		{
+			ReregisterComponent();
+		}
+	}
+	else
 	{
-		ReregisterComponent();
+		MarkRenderStateDirty();
 	}
 }
 
diff --git a/Engine/Source/Runtime/GeometryFramework/Public/Components/BaseDynamicMeshComponent.h b/Engine/Source/Runtime/GeometryFramework/Public/Components/BaseDynamicMeshComponent.h
--- a/Engine/Source/Runtime/GeometryFramework/Public/Components/BaseDynamicMeshComponent.h
+++ b/Engine/Source/Runtime/GeometryFramework/Public/Components/BaseDynamicMeshComponent.h
@@ -187,8 +187,7 @@ protected:
 	//
 public:
 	/**
-	 * if true, we always show the wireframe on top of the shaded mesh, even when not in wireframe mode
-	 * @todo: this should not be public, access via Set/Get once all usage is cleaned up
+	 * If true, render the Wireframe on top of the Shaded Mesh
 	 */
 	UPROPERTY(EditAnywhere, BlueprintReadWrite, Category = "Dynamic Mesh Component|Rendering", meta = (DisplayName = "Wireframe Overlay") )
 	bool bExplicitShowWireframe = false;
@@ -220,6 +219,20 @@ public:
 	UFUNCTION(BlueprintCallable, Category = "Dynamic Mesh Component")
 	virtual bool GetShadowsEnabled() const { return CastShadow; }
 
+	UFUNCTION(BlueprintCallable, Category = "Dynamic Mesh Component")
+	virtual void SetViewModeOverridesEnabled(bool bEnabled);
+
+	UFUNCTION(BlueprintCallable, Category = "Dynamic Mesh Component")
+	virtual bool GetViewModeOverridesEnabled() const { return bEnableViewModeOverrides; }
+
+public:
+	/** 
+	 * This flag controls whether Editor View Mode Overrides are enabled for this mesh. For example, this controls hidden-line removal on the wireframe 
+	 * in Wireframe View Mode, and whether the normal map will be disabled in Lighting-Only View Mode, as well as various other things.
+	 * Use SetViewModeOverridesEnabled() to control this setting in Blueprints/C++.
+	 */
+	UPROPERTY(EditAnywhere, BlueprintReadWrite, Category = "Dynamic Mesh Component|Rendering", meta = (DisplayName = "View Mode Overrides") )
+	bool bEnableViewModeOverrides = true;
 
 
 	//===============================================================================================================
@@ -328,6 +341,10 @@ protected:
 	// not currently support partial updates in the SceneProxy.
 public:
 
+	/**
+	 * Enable/disable Raytracing support on this Mesh, if Raytracing is currently enabled in the Project Settings.
+	 * Use SetEnableRaytracing() to configure this flag in Blueprints/C++.
+	 */
 	UPROPERTY(EditAnywhere, BlueprintReadWrite, Category = "Dynamic Mesh Component|Rendering")
 	bool bEnableRaytracing = true;
 
@@ -345,7 +362,7 @@ public:
 	virtual bool GetEnableRaytracing() const;
 
 protected:
-	virtual void OnRaytracingStateChanged();
+	virtual void OnRenderingStateChanged(bool bForceImmedateRebuild);
 
 
 
diff --git a/Engine/Source/Runtime/GeometryFramework/Public/Components/BaseDynamicMeshSceneProxy.h b/Engine/Source/Runtime/GeometryFramework/Public/Components/BaseDynamicMeshSceneProxy.h
--- a/Engine/Source/Runtime/GeometryFramework/Public/Components/BaseDynamicMeshSceneProxy.h
+++ b/Engine/Source/Runtime/GeometryFramework/Public/Components/BaseDynamicMeshSceneProxy.h
@@ -454,11 +454,15 @@ protected:
 	// control raytracing support
 	bool bEnableRaytracing = false;
 
+	// Allow view-mode overrides. 
+	bool bEnableViewModeOverrides = true;
+
 public:
 	FBaseDynamicMeshSceneProxy(UBaseDynamicMeshComponent* Component)
 		: FPrimitiveSceneProxy(Component),
 		ParentBaseComponent(Component),
-		bEnableRaytracing(Component->GetEnableRaytracing())
+		bEnableRaytracing(Component->GetEnableRaytracing()),
+		bEnableViewModeOverrides(Component->GetViewModeOverridesEnabled())
 	{
 	}
 
@@ -1047,7 +1051,7 @@ public:
 		Mesh.ReverseCulling = IsLocalToWorldDeterminantNegative();
 		Mesh.Type = PT_TriangleList;
 		Mesh.DepthPriorityGroup = DepthPriority;
-		Mesh.bCanApplyViewModeOverrides = false;
+		Mesh.bCanApplyViewModeOverrides = this->bEnableViewModeOverrides;
 		Collector.AddMesh(ViewIndex, Mesh);
 	}
 
@@ -1154,7 +1158,7 @@ public:
 		MeshBatch.ReverseCulling = IsLocalToWorldDeterminantNegative();
 		MeshBatch.Type = PT_TriangleList;
 		MeshBatch.DepthPriorityGroup = DepthPriority;
-		MeshBatch.bCanApplyViewModeOverrides = false;
+		MeshBatch.bCanApplyViewModeOverrides = this->bEnableViewModeOverrides;
 		MeshBatch.CastRayTracedShadow = IsShadowCast(Context.ReferenceView);
 
 		FMeshBatchElement& BatchElement = MeshBatch.Elements[0];
