MeshModelingTools: Fix fast update UV path for dynamic mesh proxy crashing when some UV's are unset.

#rnx

diff --git a/Engine/Source/Runtime/GeometryFramework/Public/Components/BaseDynamicMeshSceneProxy.h b/Engine/Source/Runtime/GeometryFramework/Public/Components/BaseDynamicMeshSceneProxy.h
--- a/Engine/Source/Runtime/GeometryFramework/Public/Components/BaseDynamicMeshSceneProxy.h
+++ b/Engine/Source/Runtime/GeometryFramework/Public/Components/BaseDynamicMeshSceneProxy.h
@@ -858,12 +858,23 @@ public:
 		int VertIdx = 0;
 		for (int TriangleID : Enumerable)
 		{
-			FIndex3i UVTri = UVOverlay->GetTriangle(TriangleID);
-			for (int j = 0; j < 3; ++j)
+			if (UVOverlay->IsSetTriangle(TriangleID))
 			{
-				FVector2f UV = UVOverlay->GetElement(UVTri[j]);
-				RenderBuffers->StaticMeshVertexBuffer.SetVertexUV(VertIdx, UVIndex, (FVector2D)UV);
-				VertIdx++;
+				FIndex3i UVTri = UVOverlay->GetTriangle(TriangleID);
+				for (int j = 0; j < 3; ++j)
+				{
+					FVector2f UV = UVOverlay->GetElement(UVTri[j]);
+					RenderBuffers->StaticMeshVertexBuffer.SetVertexUV(VertIdx, UVIndex, (FVector2D)UV);
+					++VertIdx;
+				}
+			}
+			else
+			{
+				for (int j = 0; j < 3; ++j)
+				{
+					RenderBuffers->StaticMeshVertexBuffer.SetVertexUV(VertIdx, UVIndex, FVector2f::Zero());
+					++VertIdx;
+				}
 			}
 		}
 	}
