UDynamicMesh: Add a bEnforceAttributeInvariants flag to the EditMesh function. If it's false, don't try to enable all required mesh attributes after the edit occurs.

Modify the DiscardMeshAttributes function in Geometry Script to not enforce attribute invariants -- otherwise the function does nothing.

#rnx

diff --git a/Engine/Plugins/Experimental/GeometryScripting/Source/GeometryScriptingCore/Private/MeshBasicEditFunctions.cpp b/Engine/Plugins/Experimental/GeometryScripting/Source/GeometryScriptingCore/Private/MeshBasicEditFunctions.cpp
--- a/Engine/Plugins/Experimental/GeometryScripting/Source/GeometryScriptingCore/Private/MeshBasicEditFunctions.cpp
+++ b/Engine/Plugins/Experimental/GeometryScripting/Source/GeometryScriptingCore/Private/MeshBasicEditFunctions.cpp
@@ -17,12 +17,13 @@ UDynamicMesh* UGeometryScriptLibrary_MeshBasicEditFunctions::DiscardMeshAttribut
 {
 	if (TargetMesh)
 	{
+		constexpr bool bEnforceAttributeInvariants = false;
 		TargetMesh->EditMesh([&](FDynamicMesh3& EditMesh)
 		{
 			EditMesh.DiscardAttributes();
 			EditMesh.DiscardVertexNormals();
 
-		}, EDynamicMeshChangeType::GeneralEdit, EDynamicMeshAttributeChangeFlags::Unknown, bDeferChangeNotifications);
+		}, EDynamicMeshChangeType::GeneralEdit, EDynamicMeshAttributeChangeFlags::Unknown, bDeferChangeNotifications, bEnforceAttributeInvariants);
 	}
 	return TargetMesh;
 }
diff --git a/Engine/Source/Runtime/GeometryFramework/Private/UDynamicMesh.cpp b/Engine/Source/Runtime/GeometryFramework/Private/UDynamicMesh.cpp
--- a/Engine/Source/Runtime/GeometryFramework/Private/UDynamicMesh.cpp
+++ b/Engine/Source/Runtime/GeometryFramework/Private/UDynamicMesh.cpp
@@ -121,15 +121,16 @@ void UDynamicMesh::ProcessMesh(TFunctionRef<void(const UE::Geometry::FDynamicMes
 void UDynamicMesh::EditMesh(TFunctionRef<void(FDynamicMesh3&)> EditFunc,
 							EDynamicMeshChangeType ChangeType,
 							EDynamicMeshAttributeChangeFlags ChangeFlags,
-							bool bDeferChangeEvents)
+							bool bDeferChangeEvents,
+							bool bEnforceAttributeInvariants)
 {
 	FDynamicMeshChangeInfo ChangeInfo;
 	ChangeInfo.Type = ChangeType;
 	ChangeInfo.Flags = ChangeFlags;
-	EditMeshInternal(EditFunc, ChangeInfo, bDeferChangeEvents);
+	EditMeshInternal(EditFunc, ChangeInfo, bDeferChangeEvents, bEnforceAttributeInvariants);
 }
 
-void UDynamicMesh::EditMeshInternal(TFunctionRef<void(FDynamicMesh3&)> EditFunc, const FDynamicMeshChangeInfo& ChangeInfo, bool bDeferChangeEvents)
+void UDynamicMesh::EditMeshInternal(TFunctionRef<void(FDynamicMesh3&)> EditFunc, const FDynamicMeshChangeInfo& ChangeInfo, bool bDeferChangeEvents, bool bEnforceAttributeInvariants)
 {
 	if (!bDeferChangeEvents)
 	{
@@ -137,18 +138,21 @@ void UDynamicMesh::EditMeshInternal(TFunctionRef<void(FDynamicMesh3&)> EditFunc,
 	}
 	EditFunc(GetMeshRef());
 
-	// Enforce our mesh attribute invariants. This should probably be optional to support compute-only UDynamicMeshes....
-	if (Mesh->HasTriangleGroups() == false)
+	if (bEnforceAttributeInvariants)
 	{
-		Mesh->EnableTriangleGroups();
-	}
-	if (Mesh->HasAttributes() == false)
-	{
-		Mesh->EnableAttributes();
-	}
-	if (Mesh->Attributes()->HasMaterialID() == false)
-	{
-		Mesh->Attributes()->EnableMaterialID();
+		// Enforce our mesh attribute invariants. Optional to support compute-only UDynamicMeshes.
+		if (Mesh->HasTriangleGroups() == false)
+		{
+			Mesh->EnableTriangleGroups();
+		}
+		if (Mesh->HasAttributes() == false)
+		{
+			Mesh->EnableAttributes();
+		}
+		if (Mesh->Attributes()->HasMaterialID() == false)
+		{
+			Mesh->Attributes()->EnableMaterialID();
+		}
 	}
 
 	if (!bDeferChangeEvents)
diff --git a/Engine/Source/Runtime/GeometryFramework/Public/UDynamicMesh.h b/Engine/Source/Runtime/GeometryFramework/Public/UDynamicMesh.h
--- a/Engine/Source/Runtime/GeometryFramework/Public/UDynamicMesh.h
+++ b/Engine/Source/Runtime/GeometryFramework/Public/UDynamicMesh.h
@@ -184,7 +184,8 @@ public:
 	void EditMesh(TFunctionRef<void(UE::Geometry::FDynamicMesh3&)> EditFunc,
 				  EDynamicMeshChangeType ChangeType = EDynamicMeshChangeType::GeneralEdit,
 				  EDynamicMeshAttributeChangeFlags ChangeFlags = EDynamicMeshAttributeChangeFlags::Unknown,
-				  bool bDeferChangeEvents = false);
+				  bool bDeferChangeEvents = false,
+				  bool bEnforceAttributeInvariants = true);
 
 	/**
 	 * Take ownership of the internal Mesh, and have it replaced with a new mesh
@@ -271,7 +272,7 @@ protected:
 	/**
 	 * Internal function that edits the Mesh, but broadcasts PreMeshChangedEvent and MeshChangedEvent, and then MeshModifiedBPEvent
 	 */
-	void EditMeshInternal(TFunctionRef<void(UE::Geometry::FDynamicMesh3&)> EditFunc, const FDynamicMeshChangeInfo& ChangeInfo, bool bDeferChangeEvents = false);
+	void EditMeshInternal(TFunctionRef<void(UE::Geometry::FDynamicMesh3&)> EditFunc, const FDynamicMeshChangeInfo& ChangeInfo, bool bDeferChangeEvents = false, bool bEnforceAttributeInvariants = true);
 
 
 public:
