MeshModelingTools: Fix DynamicMeshSceneProxy fast UV update not updating all UV channels.

#rnx

diff --git a/Engine/Source/Runtime/GeometryFramework/Private/Components/DynamicMeshSceneProxy.h b/Engine/Source/Runtime/GeometryFramework/Private/Components/DynamicMeshSceneProxy.h
--- a/Engine/Source/Runtime/GeometryFramework/Private/Components/DynamicMeshSceneProxy.h
+++ b/Engine/Source/Runtime/GeometryFramework/Private/Components/DynamicMeshSceneProxy.h
@@ -180,7 +180,7 @@ public:
 		const FDynamicMesh3* Mesh = ParentComponent->GetRenderMesh();
 
 		// find suitable overlays
-		TArray<const FDynamicMeshUVOverlay*> UVOverlays;
+		TArray<const FDynamicMeshUVOverlay*, TInlineAllocator<8>> UVOverlays;
 		const FDynamicMeshNormalOverlay* NormalOverlay = nullptr;
 		const FDynamicMeshColorOverlay* ColorOverlay = nullptr;
 		if (Mesh->HasAttributes())
@@ -227,7 +227,7 @@ public:
 		const FDynamicMeshNormalOverlay* NormalOverlay = Mesh->Attributes()->PrimaryNormals();
 		const FDynamicMeshColorOverlay* ColorOverlay = Mesh->Attributes()->PrimaryColors();
 
-		TArray<const FDynamicMeshUVOverlay*> UVOverlays;
+		TArray<const FDynamicMeshUVOverlay*, TInlineAllocator<8>> UVOverlays;
 		UVOverlays.SetNum(Attributes->NumUVLayers());
 		for (int32 k = 0; k < UVOverlays.Num(); ++k)
 		{
@@ -385,10 +385,15 @@ public:
 		{
 			NormalOverlay = Mesh->Attributes()->PrimaryNormals();
 		}
-		FDynamicMeshUVOverlay* UVOVerlay = nullptr;
+		TArray<const FDynamicMeshUVOverlay*, TInlineAllocator<8>> UVOverlays;
 		if (bUVs && ensure(Mesh->HasAttributes()) )
 		{
-			UVOVerlay = Mesh->Attributes()->PrimaryUV();
+			const FDynamicMeshAttributeSet* Attributes = Mesh->Attributes();
+			UVOverlays.SetNum(Attributes->NumUVLayers());
+			for (int32 i = 0; i < UVOverlays.Num(); ++i)
+			{
+				UVOverlays[i] = Attributes->GetUVLayer(i);
+			}
 		}
 		FDynamicMeshColorOverlay* ColorOverlay = nullptr;
 		if (bColors && ensure(Mesh->HasAttributes()) )
@@ -414,7 +419,7 @@ public:
 				if (bUVs)
 				{
 					UpdateVertexUVBufferFromOverlays(Buffers, Mesh,
-						Mesh->TriangleCount(), Mesh->TriangleIndicesItr(), UVOVerlay, 0);
+						Mesh->TriangleCount(), Mesh->TriangleIndicesItr(), UVOverlays);
 				}
 
 
@@ -446,7 +451,7 @@ public:
 					if (bUVs)
 					{
 						UpdateVertexUVBufferFromOverlays(Buffers, Mesh,
-							Buffers->Triangles->Num(), Buffers->Triangles.GetValue(), UVOVerlay, 0);
+							Buffers->Triangles->Num(), Buffers->Triangles.GetValue(), UVOverlays);
 					}
 
 					ENQUEUE_RENDER_COMMAND(FDynamicMeshSceneProxyFastUpdateVertices)(
@@ -485,10 +490,15 @@ public:
 		{
 			NormalOverlay = Mesh->Attributes()->PrimaryNormals();
 		}
-		const FDynamicMeshUVOverlay* UVOVerlay = nullptr;
+		TArray<const FDynamicMeshUVOverlay*> UVOverlays;
 		if (bUVs && ensure(Mesh->HasAttributes()))
 		{
-			UVOVerlay = Mesh->Attributes()->PrimaryUV();
+			const FDynamicMeshAttributeSet* Attributes = Mesh->Attributes();
+			UVOverlays.SetNum(Attributes->NumUVLayers());
+			for (int32 i = 0; i < UVOverlays.Num(); ++i)
+			{
+				UVOverlays[i] = Attributes->GetUVLayer(i);
+			}
 		}
 		const FDynamicMeshColorOverlay* ColorOverlay = nullptr;
 		if (bColors && ensure(Mesh->HasAttributes()))
@@ -519,7 +529,7 @@ public:
 				if (bUVs)
 				{
 					UpdateVertexUVBufferFromOverlays(Buffers, Mesh,
-						Buffers->Triangles->Num(), Buffers->Triangles.GetValue(), UVOVerlay, 0);
+						Buffers->Triangles->Num(), Buffers->Triangles.GetValue(), UVOverlays);
 				}
 
 				ENQUEUE_RENDER_COMMAND(FDynamicMeshSceneProxyFastUpdateVerticesBufferList)(
diff --git a/Engine/Source/Runtime/GeometryFramework/Public/Components/BaseDynamicMeshSceneProxy.h b/Engine/Source/Runtime/GeometryFramework/Public/Components/BaseDynamicMeshSceneProxy.h
--- a/Engine/Source/Runtime/GeometryFramework/Public/Components/BaseDynamicMeshSceneProxy.h
+++ b/Engine/Source/Runtime/GeometryFramework/Public/Components/BaseDynamicMeshSceneProxy.h
@@ -560,12 +560,12 @@ public:
 	 * Initialize rendering buffers from given attribute overlays.
 	 * Creates three vertices per triangle, IE no shared vertices in buffers.
 	 */
-	template<typename TriangleEnumerable>
+	template<typename TriangleEnumerable, typename UVOverlayListAllocator>
 	void InitializeBuffersFromOverlays(
 		FMeshRenderBufferSet* RenderBuffers,
 		const FDynamicMesh3* Mesh,
 		int NumTriangles, TriangleEnumerable Enumerable,
-		const TArray<const FDynamicMeshUVOverlay*>& UVOverlays,
+		const TArray<const FDynamicMeshUVOverlay*, UVOverlayListAllocator>& UVOverlays,
 		const FDynamicMeshNormalOverlay* NormalOverlay,
 		const FDynamicMeshColorOverlay* ColorOverlay,
 		TFunctionRef<void(int, int, int, const FVector3f&, FVector3f&, FVector3f&)> TangentsFunc,
@@ -832,56 +832,63 @@ public:
 		}
 	}
 
-
-
 	/**
 	 * Update vertex uvs of an existing set of render buffers.
 	 * Assumes that buffers were created with unshared vertices, ie three vertices per triangle, eg by InitializeBuffersFromOverlays()
 	 */
-	template<typename TriangleEnumerable>
+	template<typename TriangleEnumerable, typename UVOverlayListAllocator>
 	void UpdateVertexUVBufferFromOverlays(
 		FMeshRenderBufferSet* RenderBuffers,
 		const FDynamicMesh3* Mesh,
 		int32 NumTriangles, TriangleEnumerable Enumerable,
-		const FDynamicMeshUVOverlay* UVOverlay, int32 UVIndex)
+		const TArray<const FDynamicMeshUVOverlay*, UVOverlayListAllocator>& UVOverlays)
 	{
+		// We align the update to the way we set UV's in InitializeBuffersFromOverlays.
+
 		if (RenderBuffers->TriangleCount == 0)
 		{
 			return;
 		}
 		int NumVertices = NumTriangles * 3;
-		if ( ensure(RenderBuffers->StaticMeshVertexBuffer.GetNumVertices() == NumVertices) == false)
+		if (ensure(RenderBuffers->StaticMeshVertexBuffer.GetNumVertices() == NumVertices) == false)
 		{
 			return;
 		}
 
+		int NumUVOverlays = UVOverlays.Num();
+		int NumTexCoords = RenderBuffers->StaticMeshVertexBuffer.GetNumTexCoords();
+		if (!ensure(NumUVOverlays <= NumTexCoords))
+		{
+			return;
+		}
+
+		// Temporarily stores the UV element indices for all UV channels of a single triangle
+		TArray<FIndex3i, TFixedAllocator<MAX_STATIC_TEXCOORDS>> UVTriangles;
+		UVTriangles.SetNum(NumTexCoords);
+
 		int VertIdx = 0;
 		for (int TriangleID : Enumerable)
 		{
-			if (UVOverlay->IsSetTriangle(TriangleID))
+			for (int32 k = 0; k < NumTexCoords; ++k)
 			{
-				FIndex3i UVTri = UVOverlay->GetTriangle(TriangleID);
-				for (int j = 0; j < 3; ++j)
-				{
-					FVector2f UV = UVOverlay->GetElement(UVTri[j]);
-					RenderBuffers->StaticMeshVertexBuffer.SetVertexUV(VertIdx, UVIndex, (FVector2D)UV);
-					++VertIdx;
-				}
+				UVTriangles[k] = (k < NumUVOverlays && UVOverlays[k] != nullptr) ? UVOverlays[k]->GetTriangle(TriangleID) : FIndex3i::Invalid();
 			}
-			else
+
+			for (int j = 0; j < 3; ++j)
 			{
-				for (int j = 0; j < 3; ++j)
+				for (int32 k = 0; k < NumTexCoords; ++k)
 				{
-					RenderBuffers->StaticMeshVertexBuffer.SetVertexUV(VertIdx, UVIndex, FVector2f::Zero());
-					++VertIdx;
+					FVector2f UV = (UVTriangles[k][j] != FDynamicMesh3::InvalidID) ?
+						UVOverlays[k]->GetElement(UVTriangles[k][j]) : FVector2f::Zero();
+					RenderBuffers->StaticMeshVertexBuffer.SetVertexUV(VertIdx, k, (FVector2D)UV);
 				}
+
+				++VertIdx;
 			}
 		}
 	}
 
 
-
-
 	/**
 	 * @return number of active materials
 	 */
