ModelingComponents: disable DynamicMeshComponent in realtime sky capture by default. Results in crash in 5.0.
#rnx

diff --git a/Engine/Source/Runtime/GeometryFramework/Private/Components/BaseDynamicMeshComponent.cpp b/Engine/Source/Runtime/GeometryFramework/Private/Components/BaseDynamicMeshComponent.cpp
--- a/Engine/Source/Runtime/GeometryFramework/Private/Components/BaseDynamicMeshComponent.cpp
+++ b/Engine/Source/Runtime/GeometryFramework/Private/Components/BaseDynamicMeshComponent.cpp
@@ -8,6 +8,8 @@ using namespace UE::Geometry;
 UBaseDynamicMeshComponent::UBaseDynamicMeshComponent(const FObjectInitializer& ObjectInitializer)
 	: Super(ObjectInitializer)
 {
+	// including DynamicMeshComponent in sky captures can result in crashes in UE5.0, this should be a temporary workaround
+	bVisibleInRealTimeSkyCaptures = false;
 }
 
 
