Add an async physics build for dynamic mesh component

To support this, adds an AggGeom to the dynamic mesh component, which collision update functions should update instead of the one on the current BodySetup. This ensures new collision changes are not overwritten by an async build replacing the BodySetup.

diff --git a/Engine/Plugins/Experimental/GeometryScripting/Source/GeometryScriptingCore/Private/CollisionFunctions.cpp b/Engine/Plugins/Experimental/GeometryScripting/Source/GeometryScriptingCore/Private/CollisionFunctions.cpp
--- a/Engine/Plugins/Experimental/GeometryScripting/Source/GeometryScriptingCore/Private/CollisionFunctions.cpp
+++ b/Engine/Plugins/Experimental/GeometryScripting/Source/GeometryScriptingCore/Private/CollisionFunctions.cpp
@@ -290,29 +290,21 @@ UDynamicMesh* UGeometryScriptLibrary_CollisionFunctions::SetDynamicMeshCollision
 	}
 #endif
 
-	UBodySetup* BodySetup = DynamicMeshComponent->GetBodySetup();
-	if (BodySetup != nullptr)
-	{
-		// mark the BodySetup for modification. Do we need to modify the UStaticMesh??
 #if WITH_EDITOR
-		if (Options.bEmitTransaction)
+	if (Options.bEmitTransaction)
+	{
+		UBodySetup* BodySetup = DynamicMeshComponent->GetBodySetup();
+		if (BodySetup != nullptr)
 		{
 			BodySetup->Modify();
 		}
+	}
 #endif
 
-		// clear existing simple collision. This will call BodySetup->InvalidatePhysicsData()
-		BodySetup->RemoveSimpleCollision();
-
-		// set new collision geometry
-		BodySetup->AggGeom = NewCollision;
-
-		// update collision type
-		//BodySetup->CollisionTraceFlag = (ECollisionTraceFlag)(int32)Settings->SetCollisionType;
+	// set new collision geometry
+	DynamicMeshComponent->SetSimpleCollisionShapes(NewCollision, true /*bUpdateCollision*/);
 
-		// do we need to do a post edit change here??
-		DynamicMeshComponent->UpdateCollision(false);
-	}
+	// do we need to do a post edit change here??
 
 #if WITH_EDITOR
 	if (Options.bEmitTransaction)
@@ -344,23 +336,22 @@ void UGeometryScriptLibrary_CollisionFunctions::ResetDynamicMeshCollision(
 	}
 #endif
 
-	UBodySetup* BodySetup = DynamicMeshComponent->GetBodySetup();
-	if (BodySetup != nullptr)
-	{
-		// mark the BodySetup for modification. Do we need to modify the UStaticMesh??
 #if WITH_EDITOR
-		if (bEmitTransaction)
+	if (bEmitTransaction)
+	{
+		// mark the BodySetup for modification.
+		UBodySetup* BodySetup = DynamicMeshComponent->GetBodySetup();
+		if (BodySetup != nullptr)
 		{
 			BodySetup->Modify();
 		}
+	}
 #endif
 
-		// clear existing simple collision. This will call BodySetup->InvalidatePhysicsData()
-		BodySetup->RemoveSimpleCollision();
+	// clear existing simple collision.
+	DynamicMeshComponent->ClearSimpleCollisionShapes(true /*bUpdateCollision*/);
 
-		// do we need to do a post edit change here??
-		DynamicMeshComponent->UpdateCollision(false);
-	}
+	// do we need to do a post edit change here??
 
 #if WITH_EDITOR
 	if (bEmitTransaction)
diff --git a/Engine/Source/Runtime/GeometryFramework/Private/Components/DynamicMeshComponent.cpp b/Engine/Source/Runtime/GeometryFramework/Private/Components/DynamicMeshComponent.cpp
--- a/Engine/Source/Runtime/GeometryFramework/Private/Components/DynamicMeshComponent.cpp
+++ b/Engine/Source/Runtime/GeometryFramework/Private/Components/DynamicMeshComponent.cpp
@@ -1123,26 +1123,33 @@ bool UDynamicMeshComponent::WantsNegXTriMesh()
 	return true;
 }
 
-UBodySetup* UDynamicMeshComponent::GetBodySetup()
+UBodySetup* UDynamicMeshComponent::CreateBodySetupHelper()
 {
-	if (MeshBodySetup == nullptr)
+	UBodySetup* NewBodySetup = nullptr;
 	{
-		UBodySetup* NewBodySetup = nullptr;
-		{
-			FGCScopeGuard Scope;
+		FGCScopeGuard Scope;
 
-			// Below flags are copied from UProceduralMeshComponent::CreateBodySetupHelper(). Without these flags, DynamicMeshComponents inside
-			// a DynamicMeshActor BP will result on a GLEO error after loading and modifying a saved Level (but *not* on the initial save)
-			// The UBodySetup in a template needs to be public since the property is Instanced and thus is the archetype of the instance meaning there is a direct reference
-			NewBodySetup = NewObject<UBodySetup>(this, NAME_None, (IsTemplate() ? RF_Public | RF_ArchetypeObject : RF_NoFlags));
-		}
-		NewBodySetup->BodySetupGuid = FGuid::NewGuid();
-		
-		NewBodySetup->bGenerateMirroredCollision = false;
-		NewBodySetup->CollisionTraceFlag = this->CollisionType;
+		// Below flags are copied from UProceduralMeshComponent::CreateBodySetupHelper(). Without these flags, DynamicMeshComponents inside
+		// a DynamicMeshActor BP will result on a GLEO error after loading and modifying a saved Level (but *not* on the initial save)
+		// The UBodySetup in a template needs to be public since the property is Instanced and thus is the archetype of the instance meaning there is a direct reference
+		NewBodySetup = NewObject<UBodySetup>(this, NAME_None, (IsTemplate() ? RF_Public | RF_ArchetypeObject : RF_NoFlags));
+	}
+	NewBodySetup->BodySetupGuid = FGuid::NewGuid();
+
+	NewBodySetup->bGenerateMirroredCollision = false;
+	NewBodySetup->CollisionTraceFlag = this->CollisionType;
 
-		NewBodySetup->DefaultInstance.SetCollisionProfileName(UCollisionProfile::BlockAll_ProfileName);
-		NewBodySetup->bSupportUVsAndFaceRemap = false; /* bSupportPhysicalMaterialMasks; */
+	NewBodySetup->DefaultInstance.SetCollisionProfileName(UCollisionProfile::BlockAll_ProfileName);
+	NewBodySetup->bSupportUVsAndFaceRemap = false; /* bSupportPhysicalMaterialMasks; */
+
+	return NewBodySetup;
+}
+
+UBodySetup* UDynamicMeshComponent::GetBodySetup()
+{
+	if (MeshBodySetup == nullptr)
+	{
+		UBodySetup* NewBodySetup = CreateBodySetupHelper();
 		
 		SetBodySetup(NewBodySetup);
 	}
@@ -1158,6 +1165,24 @@ void UDynamicMeshComponent::SetBodySetup(UBodySetup* NewSetup)
 	}
 }
 
+void UDynamicMeshComponent::SetSimpleCollisionShapes(const struct FKAggregateGeom& AggGeomIn, bool bUpdateCollision)
+{
+	AggGeom = AggGeomIn;
+	if (bUpdateCollision)
+	{
+		UpdateCollision(false);
+	}
+}
+
+void UDynamicMeshComponent::ClearSimpleCollisionShapes(bool bUpdateCollision)
+{
+	AggGeom.EmptyElements();
+	if (bUpdateCollision)
+	{
+		UpdateCollision(false);
+	}
+}
+
 void UDynamicMeshComponent::InvalidatePhysicsData()
 {
 	if (GetBodySetup())
@@ -1169,16 +1194,51 @@ void UDynamicMeshComponent::InvalidatePhysicsData()
 
 void UDynamicMeshComponent::RebuildPhysicsData()
 {
-	UBodySetup* BodySetup = GetBodySetup();
-	if (BodySetup)
+	UWorld* World = GetWorld();
+	const bool bUseAsyncCook = World && World->IsGameWorld() && bUseAsyncCooking;
+
+	UBodySetup* BodySetup = nullptr;
+	if (bUseAsyncCook)
+	{
+		// Abort all previous ones still standing
+		for (UBodySetup* OldBody : AsyncBodySetupQueue)
+		{
+			OldBody->AbortPhysicsMeshAsyncCreation();
+		}
+
+		BodySetup = CreateBodySetupHelper();
+		if (BodySetup)
+		{
+			AsyncBodySetupQueue.Add(BodySetup);
+		}
+	}
+	else
+	{
+		AsyncBodySetupQueue.Empty();	// If for some reason we modified the async at runtime, just clear any pending async body setups
+		BodySetup = GetBodySetup();
+	}
+
+	if (!BodySetup)
+	{
+		return;
+	}
+
+	BodySetup->CollisionTraceFlag = this->CollisionType;
+	// Note: Directly assigning AggGeom wouldn't do some important-looking cleanup (clearing pointers on convex elements)
+	//  so we RemoveSimpleCollision then AddCollisionFrom instead
+	BodySetup->RemoveSimpleCollision();
+	BodySetup->AddCollisionFrom(this->AggGeom);
+
+	if (bUseAsyncCook)
+	{
+		BodySetup->CreatePhysicsMeshesAsync(FOnAsyncPhysicsCookFinished::CreateUObject(this, &UDynamicMeshComponent::FinishPhysicsAsyncCook, BodySetup));
+	}
+	else
 	{
 		// New GUID as collision has changed
 		BodySetup->BodySetupGuid = FGuid::NewGuid();
 		// Also we want cooked data for this
 		BodySetup->bHasCookedCollisionData = true;
-
-		BodySetup->CollisionTraceFlag = this->CollisionType;
-
 		BodySetup->InvalidatePhysicsData();
 		BodySetup->CreatePhysicsMeshes();
 		RecreatePhysicsState();
@@ -1187,6 +1247,38 @@ void UDynamicMeshComponent::RebuildPhysicsData()
 	}
 }
 
+void UDynamicMeshComponent::FinishPhysicsAsyncCook(bool bSuccess, UBodySetup* FinishedBodySetup)
+{
+	TArray<UBodySetup*> NewQueue;
+	NewQueue.Reserve(AsyncBodySetupQueue.Num());
+
+	int32 FoundIdx;
+	if (AsyncBodySetupQueue.Find(FinishedBodySetup, FoundIdx))
+	{
+		// Note: currently no-cook-needed is reported identically to cook failed.
+		// Checking AggGeom.GetElemCount() here is a hack to distinguish the no-cook-needed case
+		// TODO: remove this hack to distinguish the no-cook-needed case when/if that is no longer identical to the cook failed case
+		if (bSuccess || FinishedBodySetup->AggGeom.GetElementCount() > 0)
+		{
+			// The new body was found in the array meaning it's newer, so use it
+			MeshBodySetup = FinishedBodySetup;
+			RecreatePhysicsState();
+
+			// remove any async body setups that were requested before this one
+			for (int32 AsyncIdx = FoundIdx + 1; AsyncIdx < AsyncBodySetupQueue.Num(); ++AsyncIdx)
+			{
+				NewQueue.Add(AsyncBodySetupQueue[AsyncIdx]);
+			}
+
+			AsyncBodySetupQueue = NewQueue;
+		}
+		else
+		{
+			AsyncBodySetupQueue.RemoveAt(FoundIdx);
+		}
+	}
+}
+
 void UDynamicMeshComponent::UpdateCollision(bool bOnlyIfPending)
 {
 	if (bOnlyIfPending == false || bCollisionUpdatePending)
@@ -1195,6 +1287,13 @@ void UDynamicMeshComponent::UpdateCollision(bool bOnlyIfPending)
 	}
 }
 
+void UDynamicMeshComponent::BeginDestroy()
+{
+	Super::BeginDestroy();
+
+	AggGeom.FreeRenderInfo();
+}
+
 void UDynamicMeshComponent::EnableComplexAsSimpleCollision()
 {
 	SetComplexAsSimpleCollisionEnabled(true, true);
diff --git a/Engine/Source/Runtime/GeometryFramework/Public/Components/DynamicMeshComponent.h b/Engine/Source/Runtime/GeometryFramework/Public/Components/DynamicMeshComponent.h
--- a/Engine/Source/Runtime/GeometryFramework/Public/Components/DynamicMeshComponent.h
+++ b/Engine/Source/Runtime/GeometryFramework/Public/Components/DynamicMeshComponent.h
@@ -486,8 +486,6 @@ public:
 	virtual const UBodySetup* GetBodySetup() const { return MeshBodySetup; }
 	/** @return BodySetup for this Component. A new BodySetup will be created if one does not exist. */
 	virtual UBodySetup* GetBodySetup() override;
-	/** @return Set new BodySetup for this Component. */
-	virtual void SetBodySetup(UBodySetup* NewSetup);
 
 	/**
 	 * Force an update of the Collision/Physics data for this Component.
@@ -500,6 +498,32 @@ public:
 	UPROPERTY(EditAnywhere, BlueprintReadWrite, Category = "Dynamic Mesh Component|Collision")
 	TEnumAsByte<enum ECollisionTraceFlag> CollisionType = ECollisionTraceFlag::CTF_UseSimpleAsComplex;
 
+	/** 
+	 * Update the simple collision shapes associated with this mesh component
+	 * 
+	 * @param AggGeom				New simple collision shapes to be used
+	 * @param bUpdateCollision		Whether to automatically call UpdateCollision() -- if false, manually call it to register the change with the physics system
+	 */
+	virtual void SetSimpleCollisionShapes(const struct FKAggregateGeom& AggGeom, bool bUpdateCollision);
+
+	virtual struct FKAggregateGeom& GetSimpleCollisionShapes()
+	{
+		return AggGeom;
+	}
+
+	/**
+	 * Clear the simple collision shapes associated with this mesh component
+	 * @param bUpdateCollision		Whether to automatically call UpdateCollision() -- if false, manually call it to register the change with the physics system
+	 */
+	virtual void ClearSimpleCollisionShapes(bool bUpdateCollision);
+
+	/**
+	 *	Controls whether the physics cooking should be done off the game thread.
+	 *  This should be used when collision geometry doesn't have to be immediately up to date (For example streaming in far away objects)
+	 */
+	UPROPERTY(EditAnywhere, BlueprintReadOnly, Category = "Dynamic Mesh Component|Collision")
+	bool bUseAsyncCooking = false;
+
 	/** 
 	 * If true, current mesh will be used as Complex Collision source mesh. 
 	 * This is independent of the CollisionType setting, ie, even if Complex collision is enabled, if this is false, then the Complex Collision mesh will be empty
@@ -532,6 +556,26 @@ protected:
 	/** Recompute LocalBounds from the current Mesh */
 	void UpdateLocalBounds();
 
+	//
+	// Internals for managing collision representation and setup
+	//
+
+	/** Simplified collision representation for the mesh component  */
+	UPROPERTY(EditAnywhere, Category = BodySetup, meta = (DisplayName = "Primitives", NoResetToDefault))
+	struct FKAggregateGeom AggGeom;
+
+	/** Queue for async body setups that are being cooked */
+	UPROPERTY(transient)
+	TArray<TObjectPtr<UBodySetup>> AsyncBodySetupQueue;
+
+	/** Once async physics cook is done, create needed state */
+	virtual void FinishPhysicsAsyncCook(bool bSuccess, UBodySetup* FinishedBodySetup);
+
+	virtual UBodySetup* CreateBodySetupHelper();
+	
+	/** @return Set new BodySetup for this Component. */
+	virtual void SetBodySetup(UBodySetup* NewSetup);
+
 	/**
 	 * This is called to tell our RenderProxy about modifications to the material set.
 	 * We need to pass this on for things like material validation in the Editor.
@@ -578,6 +622,7 @@ protected:
 
 	//~ UObject Interface.
 	virtual void PostLoad() override;
+	virtual void BeginDestroy() override;
 #if WITH_EDITOR
 	void PostEditChangeProperty(FPropertyChangedEvent& PropertyChangedEvent) override;
 #endif
