Move implementations of FBaseDynamicMeshSceneProxy virtual methods from a public header to a private .cpp file

diff --git a/Engine/Source/Runtime/GeometryFramework/Private/Components/BaseDynamicMeshSceneProxy.cpp b/Engine/Source/Runtime/GeometryFramework/Private/Components/BaseDynamicMeshSceneProxy.cpp
new file mode 100644
--- /dev/null
+++ b/Engine/Source/Runtime/GeometryFramework/Private/Components/BaseDynamicMeshSceneProxy.cpp
@@ -0,0 +1,306 @@
+// Copyright Epic Games, Inc. All Rights Reserved. 
+
+#include "Components/BaseDynamicMeshSceneProxy.h"
+
+FBaseDynamicMeshSceneProxy::FBaseDynamicMeshSceneProxy(UBaseDynamicMeshComponent* Component)
+	: FPrimitiveSceneProxy(Component),
+	ParentBaseComponent(Component),
+	bEnableRaytracing(Component->GetEnableRaytracing()),
+	bEnableViewModeOverrides(Component->GetViewModeOverridesEnabled())
+{
+}
+
+FBaseDynamicMeshSceneProxy::~FBaseDynamicMeshSceneProxy()
+{
+	// we are assuming in code below that this is always called from the rendering thread
+	check(IsInRenderingThread());
+
+	// destroy all existing renderbuffers
+	for (FMeshRenderBufferSet* BufferSet : AllocatedBufferSets)
+	{
+		FMeshRenderBufferSet::DestroyRenderBufferSet(BufferSet);
+	}
+}
+
+FMeshRenderBufferSet* FBaseDynamicMeshSceneProxy::AllocateNewRenderBufferSet()
+{
+	// should we hang onto these and destroy them in constructor? leaving to subclass seems risky?
+	FMeshRenderBufferSet* RenderBufferSet = new FMeshRenderBufferSet(GetScene().GetFeatureLevel());
+
+	RenderBufferSet->Material = UMaterial::GetDefaultMaterial(MD_Surface);
+	RenderBufferSet->bEnableRaytracing = this->bEnableRaytracing;
+
+	AllocatedSetsLock.Lock();
+	AllocatedBufferSets.Add(RenderBufferSet);
+	AllocatedSetsLock.Unlock();
+
+	return RenderBufferSet;
+}
+
+void FBaseDynamicMeshSceneProxy::ReleaseRenderBufferSet(FMeshRenderBufferSet* BufferSet)
+{
+	FScopeLock Lock(&AllocatedSetsLock);
+	if (ensure(AllocatedBufferSets.Contains(BufferSet)))
+	{
+		AllocatedBufferSets.Remove(BufferSet);
+		Lock.Unlock();
+
+		FMeshRenderBufferSet::DestroyRenderBufferSet(BufferSet);
+	}
+}
+
+int32 FBaseDynamicMeshSceneProxy::GetNumMaterials() const
+{
+	return ParentBaseComponent->GetNumMaterials();
+}
+
+UMaterialInterface* FBaseDynamicMeshSceneProxy::GetMaterial(int32 k) const
+{
+	UMaterialInterface* Material = ParentBaseComponent->GetMaterial(k);
+	return (Material != nullptr) ? Material : UMaterial::GetDefaultMaterial(MD_Surface);
+}
+
+void FBaseDynamicMeshSceneProxy::UpdatedReferencedMaterials()
+{
+#if WITH_EDITOR
+	TArray<UMaterialInterface*> Materials;
+	ParentBaseComponent->GetUsedMaterials(Materials, true);
+
+	// Temporarily disable material verification while the enqueued render command is in flight.
+	// The original value for bVerifyUsedMaterials gets restored when the command is executed.
+	// If we do not do this, material verification might spuriously fail in cases where the render command for changing
+	// the verfifcation material is still in flight but the render thread is already trying to render the mesh.
+	const uint8 bRestoreVerifyUsedMaterials = bVerifyUsedMaterials;
+	bVerifyUsedMaterials = false;
+
+	ENQUEUE_RENDER_COMMAND(FMeshRenderBufferSetDestroy)(
+		[this, Materials, bRestoreVerifyUsedMaterials](FRHICommandListImmediate& RHICmdList)
+		{
+			this->SetUsedMaterialForVerification(Materials);
+			this->bVerifyUsedMaterials = bRestoreVerifyUsedMaterials;
+		});
+#endif
+}
+
+void FBaseDynamicMeshSceneProxy::GetDynamicMeshElements(const TArray<const FSceneView*>& Views, const FSceneViewFamily& ViewFamily, uint32 VisibilityMap, FMeshElementCollector& Collector) const
+{
+	QUICK_SCOPE_CYCLE_COUNTER(STAT_BaseDynamicMeshSceneProxy_GetDynamicMeshElements);
+
+	const bool bWireframe = (AllowDebugViewmodes() && ViewFamily.EngineShowFlags.Wireframe)
+		|| ParentBaseComponent->GetEnableWireframeRenderPass();
+
+	// set up wireframe material. Probably bad to reference GEngine here...also this material is very bad?
+	FMaterialRenderProxy* WireframeMaterialProxy = nullptr;
+	if (bWireframe)
+	{
+		FColoredMaterialRenderProxy* WireframeMaterialInstance = new FColoredMaterialRenderProxy(
+			GEngine->WireframeMaterial ? GEngine->WireframeMaterial->GetRenderProxy() : nullptr,
+			FLinearColor(0, 0.5f, 1.f)
+		);
+		Collector.RegisterOneFrameMaterialProxy(WireframeMaterialInstance);
+		WireframeMaterialProxy = WireframeMaterialInstance;
+	}
+
+	ESceneDepthPriorityGroup DepthPriority = SDPG_World;
+
+	TArray<FMeshRenderBufferSet*> Buffers;
+	GetActiveRenderBufferSets(Buffers);
+
+	FMaterialRenderProxy* SecondaryMaterialProxy =
+		ParentBaseComponent->HasSecondaryRenderMaterial() ? ParentBaseComponent->GetSecondaryRenderMaterial()->GetRenderProxy() : nullptr;
+	bool bDrawSecondaryBuffers = ParentBaseComponent->GetSecondaryBuffersVisibility();
+
+	for (int32 ViewIndex = 0; ViewIndex < Views.Num(); ViewIndex++)
+	{
+		if (VisibilityMap & (1 << ViewIndex))
+		{
+			const FSceneView* View = Views[ViewIndex];
+
+			bool bHasPrecomputedVolumetricLightmap;
+			FMatrix PreviousLocalToWorld;
+			int32 SingleCaptureIndex;
+			bool bOutputVelocity;
+			GetScene().GetPrimitiveUniformShaderParameters_RenderThread(GetPrimitiveSceneInfo(), bHasPrecomputedVolumetricLightmap, PreviousLocalToWorld, SingleCaptureIndex, bOutputVelocity);
+
+			// Draw the mesh.
+			for (FMeshRenderBufferSet* BufferSet : Buffers)
+			{
+				UMaterialInterface* UseMaterial = BufferSet->Material;
+				if (ParentBaseComponent->HasOverrideRenderMaterial(0))
+				{
+					UseMaterial = ParentBaseComponent->GetOverrideRenderMaterial(0);
+				}
+				FMaterialRenderProxy* MaterialProxy = UseMaterial->GetRenderProxy();
+
+				if (BufferSet->TriangleCount == 0)
+				{
+					continue;
+				}
+
+				// lock buffers so that they aren't modified while we are submitting them
+				FScopeLock BuffersLock(&BufferSet->BuffersLock);
+
+				// do we need separate one of these for each MeshRenderBufferSet?
+				FDynamicPrimitiveUniformBuffer& DynamicPrimitiveUniformBuffer = Collector.AllocateOneFrameResource<FDynamicPrimitiveUniformBuffer>();
+				DynamicPrimitiveUniformBuffer.Set(GetLocalToWorld(), PreviousLocalToWorld, GetBounds(), GetLocalBounds(), GetLocalBounds(), true, bHasPrecomputedVolumetricLightmap, DrawsVelocity(), bOutputVelocity, GetCustomPrimitiveData());
+
+				if (BufferSet->IndexBuffer.Indices.Num() > 0)
+				{
+					// Unlike most meshes, which just use the wireframe material in wireframe mode, we draw the wireframe on top of the normal material if needed,
+					// as this is easier to interpret. However, we do not do this in ortho viewports, where it frequently causes the our edit gizmo to be hidden 
+					// beneath the material. So, only draw the base material if we are in perspective mode, or we're in ortho but not in wireframe.
+					if (View->IsPerspectiveProjection() || !(AllowDebugViewmodes() && ViewFamily.EngineShowFlags.Wireframe))
+					{
+						DrawBatch(Collector, *BufferSet, BufferSet->IndexBuffer, MaterialProxy, false, DepthPriority, ViewIndex, DynamicPrimitiveUniformBuffer);
+					}
+					if (bWireframe)
+					{
+						DrawBatch(Collector, *BufferSet, BufferSet->IndexBuffer, WireframeMaterialProxy, true, DepthPriority, ViewIndex, DynamicPrimitiveUniformBuffer);
+					}
+				}
+
+				// draw secondary buffer if we have it, falling back to base material if we don't have the Secondary material
+				FMaterialRenderProxy* UseSecondaryMaterialProxy = (SecondaryMaterialProxy != nullptr) ? SecondaryMaterialProxy : MaterialProxy;
+				if (bDrawSecondaryBuffers && BufferSet->SecondaryIndexBuffer.Indices.Num() > 0 && UseSecondaryMaterialProxy != nullptr)
+				{
+					DrawBatch(Collector, *BufferSet, BufferSet->SecondaryIndexBuffer, UseSecondaryMaterialProxy, false, DepthPriority, ViewIndex, DynamicPrimitiveUniformBuffer);
+					if (bWireframe)
+					{
+						DrawBatch(Collector, *BufferSet, BufferSet->SecondaryIndexBuffer, UseSecondaryMaterialProxy, true, DepthPriority, ViewIndex, DynamicPrimitiveUniformBuffer);
+					}
+				}
+			}
+		}
+	}
+}
+
+void FBaseDynamicMeshSceneProxy::DrawBatch(FMeshElementCollector& Collector, const FMeshRenderBufferSet& RenderBuffers, const FDynamicMeshIndexBuffer32& IndexBuffer, FMaterialRenderProxy* UseMaterial, bool bWireframe, ESceneDepthPriorityGroup DepthPriority, int ViewIndex, FDynamicPrimitiveUniformBuffer& DynamicPrimitiveUniformBuffer) const
+{
+	FMeshBatch& Mesh = Collector.AllocateMesh();
+	FMeshBatchElement& BatchElement = Mesh.Elements[0];
+	BatchElement.IndexBuffer = &IndexBuffer;
+	Mesh.bWireframe = bWireframe;
+	Mesh.VertexFactory = &RenderBuffers.VertexFactory;
+	Mesh.MaterialRenderProxy = UseMaterial;
+
+	BatchElement.PrimitiveUniformBufferResource = &DynamicPrimitiveUniformBuffer.UniformBuffer;
+
+	BatchElement.FirstIndex = 0;
+	BatchElement.NumPrimitives = IndexBuffer.Indices.Num() / 3;
+	BatchElement.MinVertexIndex = 0;
+	BatchElement.MaxVertexIndex = RenderBuffers.PositionVertexBuffer.GetNumVertices() - 1;
+	Mesh.ReverseCulling = IsLocalToWorldDeterminantNegative();
+	Mesh.Type = PT_TriangleList;
+	Mesh.DepthPriorityGroup = DepthPriority;
+	Mesh.bCanApplyViewModeOverrides = this->bEnableViewModeOverrides;
+	Collector.AddMesh(ViewIndex, Mesh);
+}
+
+#if RHI_RAYTRACING
+
+bool FBaseDynamicMeshSceneProxy::IsRayTracingRelevant() const 
+{
+	return true;
+}
+
+void FBaseDynamicMeshSceneProxy::GetDynamicRayTracingInstances(FRayTracingMaterialGatheringContext& Context, TArray<FRayTracingInstance>& OutRayTracingInstances)
+{
+	QUICK_SCOPE_CYCLE_COUNTER(STAT_BaseDynamicMeshSceneProxy_GetDynamicRayTracingInstances);
+
+	ESceneDepthPriorityGroup DepthPriority = SDPG_World;
+
+	TArray<FMeshRenderBufferSet*> Buffers;
+	GetActiveRenderBufferSets(Buffers);
+
+	FMaterialRenderProxy* SecondaryMaterialProxy =
+		ParentBaseComponent->HasSecondaryRenderMaterial() ? ParentBaseComponent->GetSecondaryRenderMaterial()->GetRenderProxy() : nullptr;
+	bool bDrawSecondaryBuffers = ParentBaseComponent->GetSecondaryBuffersVisibility();
+
+	bool bHasPrecomputedVolumetricLightmap;
+	FMatrix PreviousLocalToWorld;
+	int32 SingleCaptureIndex;
+	bool bOutputVelocity;
+	GetScene().GetPrimitiveUniformShaderParameters_RenderThread(GetPrimitiveSceneInfo(), bHasPrecomputedVolumetricLightmap, PreviousLocalToWorld, SingleCaptureIndex, bOutputVelocity);
+
+	// is it safe to share this between primary and secondary raytracing batches?
+	FDynamicPrimitiveUniformBuffer& DynamicPrimitiveUniformBuffer = Context.RayTracingMeshResourceCollector.AllocateOneFrameResource<FDynamicPrimitiveUniformBuffer>();
+	DynamicPrimitiveUniformBuffer.Set(GetLocalToWorld(), PreviousLocalToWorld, GetBounds(), GetLocalBounds(), true, bHasPrecomputedVolumetricLightmap, DrawsVelocity(), bOutputVelocity);
+
+	// Draw the active buffer sets
+	for (FMeshRenderBufferSet* BufferSet : Buffers)
+	{
+		UMaterialInterface* UseMaterial = BufferSet->Material;
+		if (ParentBaseComponent->HasOverrideRenderMaterial(0))
+		{
+			UseMaterial = ParentBaseComponent->GetOverrideRenderMaterial(0);
+		}
+		FMaterialRenderProxy* MaterialProxy = UseMaterial->GetRenderProxy();
+
+		if (BufferSet->TriangleCount == 0)
+		{
+			continue;
+		}
+		if (BufferSet->bIsRayTracingDataValid == false)
+		{
+			continue;
+		}
+
+		// Lock buffers so that they aren't modified while we are submitting them.
+		FScopeLock BuffersLock(&BufferSet->BuffersLock);
+
+		// draw primary index buffer
+		if (BufferSet->IndexBuffer.Indices.Num() > 0
+			&& BufferSet->PrimaryRayTracingGeometry.RayTracingGeometryRHI.IsValid())
+		{
+			ensure(BufferSet->PrimaryRayTracingGeometry.Initializer.IndexBuffer.IsValid());
+			DrawRayTracingBatch(Context, *BufferSet, BufferSet->IndexBuffer, BufferSet->PrimaryRayTracingGeometry, MaterialProxy, DepthPriority, DynamicPrimitiveUniformBuffer, OutRayTracingInstances);
+		}
+
+		// draw secondary index buffer if we have it, falling back to base material if we don't have the Secondary material
+		FMaterialRenderProxy* UseSecondaryMaterialProxy = (SecondaryMaterialProxy != nullptr) ? SecondaryMaterialProxy : MaterialProxy;
+		if (bDrawSecondaryBuffers
+			&& BufferSet->SecondaryIndexBuffer.Indices.Num() > 0
+			&& UseSecondaryMaterialProxy != nullptr
+			&& BufferSet->SecondaryRayTracingGeometry.RayTracingGeometryRHI.IsValid())
+		{
+			ensure(BufferSet->SecondaryRayTracingGeometry.Initializer.IndexBuffer.IsValid());
+			DrawRayTracingBatch(Context, *BufferSet, BufferSet->SecondaryIndexBuffer, BufferSet->SecondaryRayTracingGeometry, UseSecondaryMaterialProxy, DepthPriority, DynamicPrimitiveUniformBuffer, OutRayTracingInstances);
+		}
+	}
+}
+
+void FBaseDynamicMeshSceneProxy::DrawRayTracingBatch(FRayTracingMaterialGatheringContext& Context, const FMeshRenderBufferSet& RenderBuffers, const FDynamicMeshIndexBuffer32& IndexBuffer, FRayTracingGeometry& RayTracingGeometry, FMaterialRenderProxy* UseMaterialProxy, ESceneDepthPriorityGroup DepthPriority, FDynamicPrimitiveUniformBuffer& DynamicPrimitiveUniformBuffer, TArray<FRayTracingInstance>& OutRayTracingInstances) const
+{
+	ensure(RayTracingGeometry.Initializer.IndexBuffer.IsValid());
+
+	FRayTracingInstance RayTracingInstance;
+	RayTracingInstance.Geometry = &RayTracingGeometry;
+	RayTracingInstance.InstanceTransforms.Add(GetLocalToWorld());
+
+	uint32 SectionIdx = 0;
+	FMeshBatch MeshBatch;
+
+	MeshBatch.VertexFactory = &RenderBuffers.VertexFactory;
+	MeshBatch.SegmentIndex = 0;
+	MeshBatch.MaterialRenderProxy = UseMaterialProxy;
+	MeshBatch.ReverseCulling = IsLocalToWorldDeterminantNegative();
+	MeshBatch.Type = PT_TriangleList;
+	MeshBatch.DepthPriorityGroup = DepthPriority;
+	MeshBatch.bCanApplyViewModeOverrides = this->bEnableViewModeOverrides;
+	MeshBatch.CastRayTracedShadow = IsShadowCast(Context.ReferenceView);
+
+	FMeshBatchElement& BatchElement = MeshBatch.Elements[0];
+	BatchElement.IndexBuffer = &IndexBuffer;
+	BatchElement.PrimitiveUniformBufferResource = &DynamicPrimitiveUniformBuffer.UniformBuffer;
+	BatchElement.FirstIndex = 0;
+	BatchElement.NumPrimitives = IndexBuffer.Indices.Num() / 3;
+	BatchElement.MinVertexIndex = 0;
+	BatchElement.MaxVertexIndex = RenderBuffers.PositionVertexBuffer.GetNumVertices() - 1;
+
+	RayTracingInstance.Materials.Add(MeshBatch);
+
+	RayTracingInstance.BuildInstanceMaskAndFlags(GetScene().GetFeatureLevel());
+	OutRayTracingInstances.Add(RayTracingInstance);
+}
+
+#endif // RHI_RAYTRACING
diff --git a/Engine/Source/Runtime/GeometryFramework/Public/Components/BaseDynamicMeshSceneProxy.h b/Engine/Source/Runtime/GeometryFramework/Public/Components/BaseDynamicMeshSceneProxy.h
--- a/Engine/Source/Runtime/GeometryFramework/Public/Components/BaseDynamicMeshSceneProxy.h
+++ b/Engine/Source/Runtime/GeometryFramework/Public/Components/BaseDynamicMeshSceneProxy.h
@@ -405,7 +405,7 @@ protected:
  * for a UBaseDynamicMeshComponent, where the assumption is that mesh data
  * will be stored in FMeshRenderBufferSet instances
  */
-class FBaseDynamicMeshSceneProxy : public FPrimitiveSceneProxy
+class GEOMETRYFRAMEWORK_API FBaseDynamicMeshSceneProxy : public FPrimitiveSceneProxy
 {
 	using FIndex2i = UE::Geometry::FIndex2i;
 	using FIndex3i = UE::Geometry::FIndex3i;
@@ -458,26 +458,9 @@ protected:
 	bool bEnableViewModeOverrides = true;
 
 public:
-	FBaseDynamicMeshSceneProxy(UBaseDynamicMeshComponent* Component)
-		: FPrimitiveSceneProxy(Component),
-		ParentBaseComponent(Component),
-		bEnableRaytracing(Component->GetEnableRaytracing()),
-		bEnableViewModeOverrides(Component->GetViewModeOverridesEnabled())
-	{
-	}
-
-
-	virtual ~FBaseDynamicMeshSceneProxy()
-	{
-		// we are assuming in code below that this is always called from the rendering thread
-		check(IsInRenderingThread());
+	FBaseDynamicMeshSceneProxy(UBaseDynamicMeshComponent* Component);
 
-		// destroy all existing renderbuffers
-		for (FMeshRenderBufferSet* BufferSet : AllocatedBufferSets)
-		{
-			FMeshRenderBufferSet::DestroyRenderBufferSet(BufferSet);
-		}
-	}
+	virtual ~FBaseDynamicMeshSceneProxy();
 
 
 	//
@@ -502,35 +485,12 @@ public:
 	 * Allocates a set of render buffers. FPrimitiveSceneProxy will keep track of these
 	 * buffers and destroy them on destruction.
 	 */
-	virtual FMeshRenderBufferSet* AllocateNewRenderBufferSet()
-	{
-		// should we hang onto these and destroy them in constructor? leaving to subclass seems risky?
-		FMeshRenderBufferSet* RenderBufferSet = new FMeshRenderBufferSet(GetScene().GetFeatureLevel());
-
-		RenderBufferSet->Material = UMaterial::GetDefaultMaterial(MD_Surface);
-		RenderBufferSet->bEnableRaytracing = this->bEnableRaytracing;
-
-		AllocatedSetsLock.Lock();
-		AllocatedBufferSets.Add(RenderBufferSet);
-		AllocatedSetsLock.Unlock();
-
-		return RenderBufferSet;
-	}
+	virtual FMeshRenderBufferSet* AllocateNewRenderBufferSet();
 
 	/**
 	 * Explicitly release a set of RenderBuffers
 	 */
-	virtual void ReleaseRenderBufferSet(FMeshRenderBufferSet* BufferSet)
-	{
-		FScopeLock Lock(&AllocatedSetsLock);
-		if (ensure(AllocatedBufferSets.Contains(BufferSet)))
-		{
-			AllocatedBufferSets.Remove(BufferSet);
-			Lock.Unlock();
-
-			FMeshRenderBufferSet::DestroyRenderBufferSet(BufferSet);
-		}
-	}
+	virtual void ReleaseRenderBufferSet(FMeshRenderBufferSet* BufferSet);
 
 
 	/**
@@ -892,19 +852,12 @@ public:
 	/**
 	 * @return number of active materials
 	 */
-	virtual int32 GetNumMaterials() const
-	{
-		return ParentBaseComponent->GetNumMaterials();
-	}
+	virtual int32 GetNumMaterials() const;
 
 	/**
 	 * Safe GetMaterial function that will never return nullptr
 	 */
-	virtual UMaterialInterface* GetMaterial(int32 k) const
-	{
-		UMaterialInterface* Material = ParentBaseComponent->GetMaterial(k);
-		return (Material != nullptr) ? Material : UMaterial::GetDefaultMaterial(MD_Surface);
-	}
+	virtual UMaterialInterface* GetMaterial(int32 k) const;
 
 	/**
 	 * Set whether or not to validate mesh batch materials against the component materials.
@@ -920,27 +873,7 @@ public:
 	 * the check in FPrimitiveSceneProxy::VerifyUsedMaterial() will fail if an override
 	 * material is set, if materials change, etc, etc
 	 */
-	virtual void UpdatedReferencedMaterials()
-	{
-#if WITH_EDITOR
-		TArray<UMaterialInterface*> Materials;
-		ParentBaseComponent->GetUsedMaterials(Materials, true);
-
-		// Temporarily disable material verification while the enqueued render command is in flight.
-		// The original value for bVerifyUsedMaterials gets restored when the command is executed.
-		// If we do not do this, material verification might spuriously fail in cases where the render command for changing
-		// the verfifcation material is still in flight but the render thread is already trying to render the mesh.
-		const uint8 bRestoreVerifyUsedMaterials = bVerifyUsedMaterials;
-		bVerifyUsedMaterials = false;
-
-		ENQUEUE_RENDER_COMMAND(FMeshRenderBufferSetDestroy)(
-			[this, Materials, bRestoreVerifyUsedMaterials](FRHICommandListImmediate& RHICmdList)
-		{
-			this->SetUsedMaterialForVerification(Materials);
-			this->bVerifyUsedMaterials = bRestoreVerifyUsedMaterials;
-		});
-#endif
-	}
+	virtual void UpdatedReferencedMaterials();
 
 
 	//
@@ -955,98 +888,7 @@ public:
 		const TArray<const FSceneView*>& Views, 
 		const FSceneViewFamily& ViewFamily, 
 		uint32 VisibilityMap, 
-		FMeshElementCollector& Collector) const override
-	{
-		QUICK_SCOPE_CYCLE_COUNTER(STAT_BaseDynamicMeshSceneProxy_GetDynamicMeshElements);
-
-		const bool bWireframe = (AllowDebugViewmodes() && ViewFamily.EngineShowFlags.Wireframe)
-			|| ParentBaseComponent->GetEnableWireframeRenderPass();
-
-		// set up wireframe material. Probably bad to reference GEngine here...also this material is very bad?
-		FMaterialRenderProxy* WireframeMaterialProxy = nullptr;
-		if (bWireframe)
-		{
-			FColoredMaterialRenderProxy* WireframeMaterialInstance = new FColoredMaterialRenderProxy(
-				GEngine->WireframeMaterial ? GEngine->WireframeMaterial->GetRenderProxy() : nullptr,
-				FLinearColor(0, 0.5f, 1.f)
-			);
-			Collector.RegisterOneFrameMaterialProxy(WireframeMaterialInstance);
-			WireframeMaterialProxy = WireframeMaterialInstance;
-		}
-
-		ESceneDepthPriorityGroup DepthPriority = SDPG_World;
-
-		TArray<FMeshRenderBufferSet*> Buffers;
-		GetActiveRenderBufferSets(Buffers);
-
-		FMaterialRenderProxy* SecondaryMaterialProxy =
-			ParentBaseComponent->HasSecondaryRenderMaterial() ? ParentBaseComponent->GetSecondaryRenderMaterial()->GetRenderProxy() : nullptr;
-		bool bDrawSecondaryBuffers = ParentBaseComponent->GetSecondaryBuffersVisibility();
-
-		for (int32 ViewIndex = 0; ViewIndex < Views.Num(); ViewIndex++)
-		{
-			if (VisibilityMap & (1 << ViewIndex))
-			{
-				const FSceneView* View = Views[ViewIndex];
-
-				bool bHasPrecomputedVolumetricLightmap;
-				FMatrix PreviousLocalToWorld;
-				int32 SingleCaptureIndex;
-				bool bOutputVelocity;
-				GetScene().GetPrimitiveUniformShaderParameters_RenderThread(GetPrimitiveSceneInfo(), bHasPrecomputedVolumetricLightmap, PreviousLocalToWorld, SingleCaptureIndex, bOutputVelocity);
-
-				// Draw the mesh.
-				for (FMeshRenderBufferSet* BufferSet : Buffers)
-				{
-					UMaterialInterface* UseMaterial = BufferSet->Material;
-					if (ParentBaseComponent->HasOverrideRenderMaterial(0))
-					{
-						UseMaterial = ParentBaseComponent->GetOverrideRenderMaterial(0);
-					}
-					FMaterialRenderProxy* MaterialProxy = UseMaterial->GetRenderProxy();
-
-					if (BufferSet->TriangleCount == 0)
-					{
-						continue;
-					}
-
-					// lock buffers so that they aren't modified while we are submitting them
-					FScopeLock BuffersLock(&BufferSet->BuffersLock);
-
-					// do we need separate one of these for each MeshRenderBufferSet?
-					FDynamicPrimitiveUniformBuffer& DynamicPrimitiveUniformBuffer = Collector.AllocateOneFrameResource<FDynamicPrimitiveUniformBuffer>();
-					DynamicPrimitiveUniformBuffer.Set(GetLocalToWorld(), PreviousLocalToWorld, GetBounds(), GetLocalBounds(), GetLocalBounds(), true, bHasPrecomputedVolumetricLightmap, DrawsVelocity(), bOutputVelocity, GetCustomPrimitiveData());
-
-					if (BufferSet->IndexBuffer.Indices.Num() > 0)
-					{
-						// Unlike most meshes, which just use the wireframe material in wireframe mode, we draw the wireframe on top of the normal material if needed,
-						// as this is easier to interpret. However, we do not do this in ortho viewports, where it frequently causes the our edit gizmo to be hidden 
-						// beneath the material. So, only draw the base material if we are in perspective mode, or we're in ortho but not in wireframe.
-						if (View->IsPerspectiveProjection() || !(AllowDebugViewmodes() && ViewFamily.EngineShowFlags.Wireframe))
-						{
-							DrawBatch(Collector, *BufferSet, BufferSet->IndexBuffer, MaterialProxy, false, DepthPriority, ViewIndex, DynamicPrimitiveUniformBuffer);
-						}
-						if (bWireframe)
-						{
-							DrawBatch(Collector, *BufferSet, BufferSet->IndexBuffer, WireframeMaterialProxy, true, DepthPriority, ViewIndex, DynamicPrimitiveUniformBuffer);
-						}
-					}
-
-					// draw secondary buffer if we have it, falling back to base material if we don't have the Secondary material
-					FMaterialRenderProxy* UseSecondaryMaterialProxy = (SecondaryMaterialProxy != nullptr) ? SecondaryMaterialProxy : MaterialProxy;
-					if (bDrawSecondaryBuffers && BufferSet->SecondaryIndexBuffer.Indices.Num() > 0 && UseSecondaryMaterialProxy != nullptr)
-					{
-						DrawBatch(Collector, *BufferSet, BufferSet->SecondaryIndexBuffer, UseSecondaryMaterialProxy, false, DepthPriority, ViewIndex, DynamicPrimitiveUniformBuffer);
-						if (bWireframe)
-						{
-							DrawBatch(Collector, *BufferSet, BufferSet->SecondaryIndexBuffer, UseSecondaryMaterialProxy, true, DepthPriority, ViewIndex, DynamicPrimitiveUniformBuffer);
-						}
-					}
-				}
-			}
-		}
-	}
-
+		FMeshElementCollector& Collector) const override;
 
 
 	/**
@@ -1059,101 +901,14 @@ public:
 		bool bWireframe,
 		ESceneDepthPriorityGroup DepthPriority,
 		int ViewIndex,
-		FDynamicPrimitiveUniformBuffer& DynamicPrimitiveUniformBuffer) const
-	{
-		FMeshBatch& Mesh = Collector.AllocateMesh();
-		FMeshBatchElement& BatchElement = Mesh.Elements[0];
-		BatchElement.IndexBuffer = &IndexBuffer;
-		Mesh.bWireframe = bWireframe;
-		Mesh.VertexFactory = &RenderBuffers.VertexFactory;
-		Mesh.MaterialRenderProxy = UseMaterial;
-
-		BatchElement.PrimitiveUniformBufferResource = &DynamicPrimitiveUniformBuffer.UniformBuffer;
-
-		BatchElement.FirstIndex = 0;
-		BatchElement.NumPrimitives = IndexBuffer.Indices.Num() / 3;
-		BatchElement.MinVertexIndex = 0;
-		BatchElement.MaxVertexIndex = RenderBuffers.PositionVertexBuffer.GetNumVertices() - 1;
-		Mesh.ReverseCulling = IsLocalToWorldDeterminantNegative();
-		Mesh.Type = PT_TriangleList;
-		Mesh.DepthPriorityGroup = DepthPriority;
-		Mesh.bCanApplyViewModeOverrides = this->bEnableViewModeOverrides;
-		Collector.AddMesh(ViewIndex, Mesh);
-	}
-
-
-
+		FDynamicPrimitiveUniformBuffer& DynamicPrimitiveUniformBuffer) const;
 
 #if RHI_RAYTRACING
 
-	virtual bool IsRayTracingRelevant() const override { return true; }
-
-	virtual void GetDynamicRayTracingInstances(FRayTracingMaterialGatheringContext& Context, TArray<FRayTracingInstance>& OutRayTracingInstances) override
-	{
-		QUICK_SCOPE_CYCLE_COUNTER(STAT_BaseDynamicMeshSceneProxy_GetDynamicRayTracingInstances);
-
-		ESceneDepthPriorityGroup DepthPriority = SDPG_World;
-
-		TArray<FMeshRenderBufferSet*> Buffers;
-		GetActiveRenderBufferSets(Buffers);
-
-		FMaterialRenderProxy* SecondaryMaterialProxy =
-			ParentBaseComponent->HasSecondaryRenderMaterial() ? ParentBaseComponent->GetSecondaryRenderMaterial()->GetRenderProxy() : nullptr;
-		bool bDrawSecondaryBuffers = ParentBaseComponent->GetSecondaryBuffersVisibility();
-
-		bool bHasPrecomputedVolumetricLightmap;
-		FMatrix PreviousLocalToWorld;
-		int32 SingleCaptureIndex;
-		bool bOutputVelocity;
-		GetScene().GetPrimitiveUniformShaderParameters_RenderThread(GetPrimitiveSceneInfo(), bHasPrecomputedVolumetricLightmap, PreviousLocalToWorld, SingleCaptureIndex, bOutputVelocity);
-
-		// is it safe to share this between primary and secondary raytracing batches?
-		FDynamicPrimitiveUniformBuffer& DynamicPrimitiveUniformBuffer = Context.RayTracingMeshResourceCollector.AllocateOneFrameResource<FDynamicPrimitiveUniformBuffer>();
-		DynamicPrimitiveUniformBuffer.Set(GetLocalToWorld(), PreviousLocalToWorld, GetBounds(), GetLocalBounds(), true, bHasPrecomputedVolumetricLightmap, DrawsVelocity(), bOutputVelocity);
-
-		// Draw the active buffer sets
-		for (FMeshRenderBufferSet* BufferSet : Buffers)
-		{
-			UMaterialInterface* UseMaterial = BufferSet->Material;
-			if (ParentBaseComponent->HasOverrideRenderMaterial(0))
-			{
-				UseMaterial = ParentBaseComponent->GetOverrideRenderMaterial(0);
-			}
-			FMaterialRenderProxy* MaterialProxy = UseMaterial->GetRenderProxy();
-
-			if (BufferSet->TriangleCount == 0)
-			{
-				continue;
-			}
-			if (BufferSet->bIsRayTracingDataValid == false)
-			{
-				continue;
-			}
+	virtual bool IsRayTracingRelevant() const override;
 
-			// Lock buffers so that they aren't modified while we are submitting them.
-			FScopeLock BuffersLock(&BufferSet->BuffersLock);
-
-			// draw primary index buffer
-			if (BufferSet->IndexBuffer.Indices.Num() > 0 
-				&& BufferSet->PrimaryRayTracingGeometry.RayTracingGeometryRHI.IsValid())
-			{
-				ensure(BufferSet->PrimaryRayTracingGeometry.Initializer.IndexBuffer.IsValid());
-				DrawRayTracingBatch(Context, *BufferSet, BufferSet->IndexBuffer, BufferSet->PrimaryRayTracingGeometry, MaterialProxy, DepthPriority, DynamicPrimitiveUniformBuffer, OutRayTracingInstances);
-			}
-
-			// draw secondary index buffer if we have it, falling back to base material if we don't have the Secondary material
-			FMaterialRenderProxy* UseSecondaryMaterialProxy = (SecondaryMaterialProxy != nullptr) ? SecondaryMaterialProxy : MaterialProxy;
-			if (bDrawSecondaryBuffers 
-				&& BufferSet->SecondaryIndexBuffer.Indices.Num() > 0 
-				&& UseSecondaryMaterialProxy != nullptr
-				&& BufferSet->SecondaryRayTracingGeometry.RayTracingGeometryRHI.IsValid() )
-			{
-				ensure(BufferSet->SecondaryRayTracingGeometry.Initializer.IndexBuffer.IsValid());
-				DrawRayTracingBatch(Context, *BufferSet, BufferSet->SecondaryIndexBuffer, BufferSet->SecondaryRayTracingGeometry, UseSecondaryMaterialProxy, DepthPriority, DynamicPrimitiveUniformBuffer, OutRayTracingInstances);
-			}
-		}
-	}
 
+	virtual void GetDynamicRayTracingInstances(FRayTracingMaterialGatheringContext& Context, TArray<FRayTracingInstance>& OutRayTracingInstances) override;
 
 
 	/**
@@ -1167,46 +922,9 @@ public:
 		FMaterialRenderProxy* UseMaterialProxy,
 		ESceneDepthPriorityGroup DepthPriority,
 		FDynamicPrimitiveUniformBuffer& DynamicPrimitiveUniformBuffer,
-		TArray<FRayTracingInstance>& OutRayTracingInstances	) const
-	{
-		ensure(RayTracingGeometry.Initializer.IndexBuffer.IsValid());
-
-		FRayTracingInstance RayTracingInstance;
-		RayTracingInstance.Geometry = &RayTracingGeometry;
-		RayTracingInstance.InstanceTransforms.Add(GetLocalToWorld());
-
-		uint32 SectionIdx = 0;
-		FMeshBatch MeshBatch;
-
-		MeshBatch.VertexFactory = &RenderBuffers.VertexFactory;
-		MeshBatch.SegmentIndex = 0;
-		MeshBatch.MaterialRenderProxy = UseMaterialProxy;
-		MeshBatch.ReverseCulling = IsLocalToWorldDeterminantNegative();
-		MeshBatch.Type = PT_TriangleList;
-		MeshBatch.DepthPriorityGroup = DepthPriority;
-		MeshBatch.bCanApplyViewModeOverrides = this->bEnableViewModeOverrides;
-		MeshBatch.CastRayTracedShadow = IsShadowCast(Context.ReferenceView);
-
-		FMeshBatchElement& BatchElement = MeshBatch.Elements[0];
-		BatchElement.IndexBuffer = &IndexBuffer;
-		BatchElement.PrimitiveUniformBufferResource = &DynamicPrimitiveUniformBuffer.UniformBuffer;
-		BatchElement.FirstIndex = 0;
-		BatchElement.NumPrimitives = IndexBuffer.Indices.Num() / 3;
-		BatchElement.MinVertexIndex = 0;
-		BatchElement.MaxVertexIndex = RenderBuffers.PositionVertexBuffer.GetNumVertices() - 1;
-
-		RayTracingInstance.Materials.Add(MeshBatch);
-
-		RayTracingInstance.BuildInstanceMaskAndFlags(GetScene().GetFeatureLevel());
-		OutRayTracingInstances.Add(RayTracingInstance);
-	}
+		TArray<FRayTracingInstance>& OutRayTracingInstances) const;
 
 
 #endif // RHI_RAYTRACING
 
-
-
-
-
-
-};
\ No newline at end of file
+};
