Remove few unnecessary includes from BaseDynamicMeshSceneProxy.h

diff --git a/Engine/Source/Runtime/GeometryFramework/Private/Components/BaseDynamicMeshSceneProxy.cpp b/Engine/Source/Runtime/GeometryFramework/Private/Components/BaseDynamicMeshSceneProxy.cpp
--- a/Engine/Source/Runtime/GeometryFramework/Private/Components/BaseDynamicMeshSceneProxy.cpp
+++ b/Engine/Source/Runtime/GeometryFramework/Private/Components/BaseDynamicMeshSceneProxy.cpp
@@ -1,6 +1,9 @@
 // Copyright Epic Games, Inc. All Rights Reserved. 
 
 #include "Components/BaseDynamicMeshSceneProxy.h"
+#include "Materials/Material.h"
+#include "RayTracingDefinitions.h"
+#include "RayTracingInstance.h"
 
 FBaseDynamicMeshSceneProxy::FBaseDynamicMeshSceneProxy(UBaseDynamicMeshComponent* Component)
 	: FPrimitiveSceneProxy(Component),
diff --git a/Engine/Source/Runtime/GeometryFramework/Public/Components/BaseDynamicMeshSceneProxy.h b/Engine/Source/Runtime/GeometryFramework/Public/Components/BaseDynamicMeshSceneProxy.h
--- a/Engine/Source/Runtime/GeometryFramework/Public/Components/BaseDynamicMeshSceneProxy.h
+++ b/Engine/Source/Runtime/GeometryFramework/Public/Components/BaseDynamicMeshSceneProxy.h
@@ -12,9 +12,6 @@
 #include "Rendering/ColorVertexBuffer.h"
 #include "DynamicMeshBuilder.h"
 #include "Components/BaseDynamicMeshComponent.h"
-#include "Materials/Material.h"
-#include "RayTracingDefinitions.h"
-#include "RayTracingInstance.h"
 
 using UE::Geometry::FDynamicMesh3;
 using UE::Geometry::FDynamicMeshAttributeSet;
@@ -23,6 +20,10 @@ using UE::Geometry::FDynamicMeshNormalOverlay;
 using UE::Geometry::FDynamicMeshColorOverlay;
 using UE::Geometry::FDynamicMeshMaterialAttribute;
 
+class UMaterialInterface;
+class FMaterialRenderProxy;
+struct FRayTracingMaterialGatheringContext;
+
 /**
  * FMeshRenderBufferSet stores a set of RenderBuffers for a mesh
  */
