GeometryFramework: add support for configuring color and normal overrides on BaseDynamicMeshComponent & DynamicMeshComponent. This allows the user to configure a DynamicMeshComponent to display vertex colors, group colors, or a (configurable) custom color, as well as facet normals, in the Actor properties.

A new lit vertex color material is added and used if the ModelingComponents module is loaded (to avoid adding a new engine material, if only GeometryFramework is loaded, then the engine-default vertex color material is used). This class-wide vertex color material can be overridden via UBaseDynamicMeshComponent::SetDefaultVertexColorMaterial()

Facet Normals / Flat Shading is now separately configurable, this overrides the normals at the vertexbuffer setup level, so it works independently of material.

Support also added for configuring the wireframe material used for BaseDynamicMeshComponent, and also for setting the wireframe color.

diff --git a/Engine/Plugins/Runtime/MeshModelingToolset/Source/ModelingComponents/Private/ModelingComponentsModule.cpp b/Engine/Plugins/Runtime/MeshModelingToolset/Source/ModelingComponents/Private/ModelingComponentsModule.cpp
--- a/Engine/Plugins/Runtime/MeshModelingToolset/Source/ModelingComponents/Private/ModelingComponentsModule.cpp
+++ b/Engine/Plugins/Runtime/MeshModelingToolset/Source/ModelingComponents/Private/ModelingComponentsModule.cpp
@@ -1,18 +1,34 @@
 // Copyright Epic Games, Inc. All Rights Reserved.
 
 #include "ModelingComponentsModule.h"
+#include "Components/BaseDynamicMeshComponent.h"
+
 
 #define LOCTEXT_NAMESPACE "FModelingComponentsModule"
 
 void FModelingComponentsModule::StartupModule()
 {
-	// This code will execute after your module is loaded into memory; the exact timing is specified in the .uplugin file per-module
+	FCoreDelegates::OnPostEngineInit.AddRaw(this, &FModelingComponentsModule::OnPostEngineInit);
+
+	// Ensure that the GeometryFramework module is loaded so that UBaseDynamicMeshComponent materials are configured
+	FModuleManager::Get().LoadModule(TEXT("GeometryFramework"));
+}
+
+void FModelingComponentsModule::OnPostEngineInit()
+{
+	// Replace the standard UBaseDynamicMeshComponent vertex color material with something higher quality.
+	// This is done in ModelingComponents module (ie part of MeshModelingToolset plugin) to avoid having to
+	// make this Material a special "engine material", which has various undesirable implication
+	UMaterial* VertexColorMaterial = LoadObject<UMaterial>(nullptr, TEXT("/MeshModelingToolset/Materials/M_DynamicMeshComponentVtxColor"));
+	if (ensure(VertexColorMaterial))
+	{
+		UBaseDynamicMeshComponent::SetDefaultVertexColorMaterial(VertexColorMaterial);
+	}
 }
 
 void FModelingComponentsModule::ShutdownModule()
 {
-	// This function may be called during shutdown to clean up your module.  For modules that support dynamic reloading,
-	// we call this function before unloading the module.
+	FCoreDelegates::OnPostEngineInit.RemoveAll(this);
 }
 
 #undef LOCTEXT_NAMESPACE
diff --git a/Engine/Plugins/Runtime/MeshModelingToolset/Source/ModelingComponents/Public/ModelingComponentsModule.h b/Engine/Plugins/Runtime/MeshModelingToolset/Source/ModelingComponents/Public/ModelingComponentsModule.h
--- a/Engine/Plugins/Runtime/MeshModelingToolset/Source/ModelingComponents/Public/ModelingComponentsModule.h
+++ b/Engine/Plugins/Runtime/MeshModelingToolset/Source/ModelingComponents/Public/ModelingComponentsModule.h
@@ -12,4 +12,7 @@ public:
 	/** IModuleInterface implementation */
 	virtual void StartupModule() override;
 	virtual void ShutdownModule() override;
+
+	// registered in StartupModule() with FCoreDelegates::OnPostEngineInit and called once
+	void OnPostEngineInit();
 };
diff --git a/Engine/Source/Editor/UnrealEd/Private/EditorEngine.cpp b/Engine/Source/Editor/UnrealEd/Private/EditorEngine.cpp
--- a/Engine/Source/Editor/UnrealEd/Private/EditorEngine.cpp
+++ b/Engine/Source/Editor/UnrealEd/Private/EditorEngine.cpp
@@ -1065,6 +1065,7 @@ void UEditorEngine::Init(IEngineLoop* InEngineLoop)
 			TEXT("EditorConfig"),
 			TEXT("DerivedDataEditor"),
 			TEXT("CSVtoSVG"),
+			TEXT("GeometryFramework")
 		};
 
 		FScopedSlowTask ModuleSlowTask(UE_ARRAY_COUNT(ModuleNames));
diff --git a/Engine/Source/Runtime/GeometryFramework/Private/Components/BaseDynamicMeshComponent.cpp b/Engine/Source/Runtime/GeometryFramework/Private/Components/BaseDynamicMeshComponent.cpp
--- a/Engine/Source/Runtime/GeometryFramework/Private/Components/BaseDynamicMeshComponent.cpp
+++ b/Engine/Source/Runtime/GeometryFramework/Private/Components/BaseDynamicMeshComponent.cpp
@@ -2,6 +2,7 @@
 
 #include "Components/BaseDynamicMeshComponent.h"
 #include "Components/BaseDynamicMeshSceneProxy.h"
+#include "Materials/Material.h"
 
 using namespace UE::Geometry;
 
@@ -115,6 +116,38 @@ bool UBaseDynamicMeshComponent::GetEnableRaytracing() const
 	return bEnableRaytracing;
 }
 
+
+
+void UBaseDynamicMeshComponent::SetColorOverrideMode(EDynamicMeshComponentColorOverrideMode NewMode)
+{
+	if (ColorMode != NewMode)
+	{
+		ColorMode = NewMode;
+		OnRenderingStateChanged(false);
+	}
+}
+
+
+void UBaseDynamicMeshComponent::SetConstantOverrideColor(FColor NewColor) 
+{ 
+	if (ConstantColor != NewColor)
+	{
+		ConstantColor = NewColor; 
+		OnRenderingStateChanged(false);
+	}
+}
+
+
+void UBaseDynamicMeshComponent::SetEnableFlatShading(bool bEnable)
+{
+	if (bEnableFlatShading != bEnable)
+	{
+		bEnableFlatShading = bEnable; 
+		OnRenderingStateChanged(false);
+	}
+}
+
+
 void UBaseDynamicMeshComponent::OnRenderingStateChanged(bool bForceImmedateRebuild)
 {
 	if (bForceImmedateRebuild)
@@ -198,5 +231,52 @@ void UBaseDynamicMeshComponent::GetUsedMaterials(TArray<UMaterialInterface*>& Ou
 	{
 		OutMaterials.Add(SecondaryRenderMaterial);
 	}
+	if (ColorMode != EDynamicMeshComponentColorOverrideMode::None && UBaseDynamicMeshComponent::DefaultVertexColorMaterial != nullptr)
+	{
+		OutMaterials.Add(UBaseDynamicMeshComponent::DefaultVertexColorMaterial);
+	}
+}
+
+
+UMaterialInterface* UBaseDynamicMeshComponent::DefaultWireframeMaterial = nullptr;
+UMaterialInterface* UBaseDynamicMeshComponent::DefaultVertexColorMaterial = nullptr;
+
+void UBaseDynamicMeshComponent::InitializeDefaultMaterials()
+{
+	if (GEngine->WireframeMaterial != nullptr)
+	{
+		DefaultWireframeMaterial = GEngine->WireframeMaterial;
+	}
+	if (GEngine->VertexColorViewModeMaterial_ColorOnly != nullptr)
+	{
+		DefaultVertexColorMaterial = GEngine->VertexColorViewModeMaterial_ColorOnly;
+	}
+}
+
+UMaterialInterface* UBaseDynamicMeshComponent::GetDefaultWireframeMaterial_RenderThread()
+{
+	return DefaultWireframeMaterial;
+}
+
+void UBaseDynamicMeshComponent::SetDefaultWireframeMaterial(UMaterialInterface* Material)
+{
+	ENQUEUE_RENDER_COMMAND(BaseDynamicMeshComponent_SetWireframeMaterial)(
+		[Material](FRHICommandListImmediate& RHICmdList)
+	{
+		UBaseDynamicMeshComponent::DefaultWireframeMaterial = Material;
+	});
 }
 
+UMaterialInterface* UBaseDynamicMeshComponent::GetDefaultVertexColorMaterial_RenderThread()
+{
+	return DefaultVertexColorMaterial;
+}
+
+void UBaseDynamicMeshComponent::SetDefaultVertexColorMaterial(UMaterialInterface* Material)
+{
+	ENQUEUE_RENDER_COMMAND(BaseDynamicMeshComponent_SetVertexColorMaterial)(
+		[Material](FRHICommandListImmediate& RHICmdList)
+	{
+		UBaseDynamicMeshComponent::DefaultVertexColorMaterial = Material;
+	});
+}
\ No newline at end of file
diff --git a/Engine/Source/Runtime/GeometryFramework/Private/Components/BaseDynamicMeshSceneProxy.cpp b/Engine/Source/Runtime/GeometryFramework/Private/Components/BaseDynamicMeshSceneProxy.cpp
--- a/Engine/Source/Runtime/GeometryFramework/Private/Components/BaseDynamicMeshSceneProxy.cpp
+++ b/Engine/Source/Runtime/GeometryFramework/Private/Components/BaseDynamicMeshSceneProxy.cpp
@@ -11,6 +11,12 @@ FBaseDynamicMeshSceneProxy::FBaseDynamicMeshSceneProxy(UBaseDynamicMeshComponent
 	bEnableRaytracing(Component->GetEnableRaytracing()),
 	bEnableViewModeOverrides(Component->GetViewModeOverridesEnabled())
 {
+	if (Component->GetColorOverrideMode() == EDynamicMeshComponentColorOverrideMode::Constant)
+	{
+		ConstantVertexColor = Component->GetConstantOverrideColor();
+	}
+
+	bUsePerTriangleNormals = Component->GetFlatShadingEnabled();
 }
 
 FBaseDynamicMeshSceneProxy::~FBaseDynamicMeshSceneProxy()
@@ -78,10 +84,10 @@ void FBaseDynamicMeshSceneProxy::UpdatedReferencedMaterials()
 
 	ENQUEUE_RENDER_COMMAND(FMeshRenderBufferSetDestroy)(
 		[this, Materials, bRestoreVerifyUsedMaterials](FRHICommandListImmediate& RHICmdList)
-		{
-			this->SetUsedMaterialForVerification(Materials);
-			this->bVerifyUsedMaterials = bRestoreVerifyUsedMaterials;
-		});
+	{
+		this->SetUsedMaterialForVerification(Materials);
+		this->bVerifyUsedMaterials = bRestoreVerifyUsedMaterials;
+	});
 #endif
 }
 
@@ -89,19 +95,39 @@ void FBaseDynamicMeshSceneProxy::GetDynamicMeshElements(const TArray<const FScen
 {
 	QUICK_SCOPE_CYCLE_COUNTER(STAT_BaseDynamicMeshSceneProxy_GetDynamicMeshElements);
 
-	const bool bWireframe = (AllowDebugViewmodes() && ViewFamily.EngineShowFlags.Wireframe)
+	bool bWireframe = (AllowDebugViewmodes() && ViewFamily.EngineShowFlags.Wireframe)
 		|| ParentBaseComponent->GetEnableWireframeRenderPass();
 
-	// set up wireframe material. Probably bad to reference GEngine here...also this material is very bad?
+	// Get wireframe material proxy if requested and available, otherwise disable wireframe
 	FMaterialRenderProxy* WireframeMaterialProxy = nullptr;
 	if (bWireframe)
 	{
-		FColoredMaterialRenderProxy* WireframeMaterialInstance = new FColoredMaterialRenderProxy(
-			GEngine->WireframeMaterial ? GEngine->WireframeMaterial->GetRenderProxy() : nullptr,
-			FLinearColor(0, 0.5f, 1.f)
-		);
-		Collector.RegisterOneFrameMaterialProxy(WireframeMaterialInstance);
-		WireframeMaterialProxy = WireframeMaterialInstance;
+		UMaterialInterface* WireframeMaterial = UBaseDynamicMeshComponent::GetDefaultWireframeMaterial_RenderThread();
+		if (WireframeMaterial != nullptr)
+		{
+			FColoredMaterialRenderProxy* WireframeMaterialInstance = new FColoredMaterialRenderProxy(
+				WireframeMaterial->GetRenderProxy(), ParentBaseComponent->WireframeColor);
+			Collector.RegisterOneFrameMaterialProxy(WireframeMaterialInstance);
+			WireframeMaterialProxy = WireframeMaterialInstance;
+		}
+		else
+		{
+			bWireframe = false;
+		}
+	}
+
+	// will use this material instead of any others below, if it becomes non-null
+	UMaterialInterface* ForceOverrideMaterial = nullptr;
+
+	// Note: there is an engine show flags for vertex colors. However if we enable that, then
+	// we will fail the used-materials checks in the Editor, because UBaseDynamicMeshComponent::GetUsedMaterials()
+	// cannot check the active show flags. One option would be to always add the vertex color material there...
+	const bool bVertexColor = ParentBaseComponent->ColorMode == EDynamicMeshComponentColorOverrideMode::VertexColors ||
+								ParentBaseComponent->ColorMode == EDynamicMeshComponentColorOverrideMode::Polygroups ||
+								ParentBaseComponent->ColorMode == EDynamicMeshComponentColorOverrideMode::Constant;
+	if (bVertexColor)
+	{
+		ForceOverrideMaterial = UBaseDynamicMeshComponent::GetDefaultVertexColorMaterial_RenderThread();
 	}
 
 	ESceneDepthPriorityGroup DepthPriority = SDPG_World;
@@ -109,8 +135,11 @@ void FBaseDynamicMeshSceneProxy::GetDynamicMeshElements(const TArray<const FScen
 	TArray<FMeshRenderBufferSet*> Buffers;
 	GetActiveRenderBufferSets(Buffers);
 
-	FMaterialRenderProxy* SecondaryMaterialProxy =
-		ParentBaseComponent->HasSecondaryRenderMaterial() ? ParentBaseComponent->GetSecondaryRenderMaterial()->GetRenderProxy() : nullptr;
+	UMaterialInterface* UseSecondaryMaterial = ForceOverrideMaterial;
+	if (ParentBaseComponent->HasSecondaryRenderMaterial() && ForceOverrideMaterial == nullptr)
+	{
+		UseSecondaryMaterial = ParentBaseComponent->GetSecondaryRenderMaterial();
+	}
 	bool bDrawSecondaryBuffers = ParentBaseComponent->GetSecondaryBuffersVisibility();
 
 	for (int32 ViewIndex = 0; ViewIndex < Views.Num(); ViewIndex++)
@@ -134,6 +163,10 @@ void FBaseDynamicMeshSceneProxy::GetDynamicMeshElements(const TArray<const FScen
 				{
 					UseMaterial = ParentBaseComponent->GetOverrideRenderMaterial(0);
 				}
+				if (ForceOverrideMaterial)
+				{
+					UseMaterial = ForceOverrideMaterial;
+				}
 				FMaterialRenderProxy* MaterialProxy = UseMaterial->GetRenderProxy();
 
 				if (BufferSet->TriangleCount == 0)
@@ -164,7 +197,7 @@ void FBaseDynamicMeshSceneProxy::GetDynamicMeshElements(const TArray<const FScen
 				}
 
 				// draw secondary buffer if we have it, falling back to base material if we don't have the Secondary material
-				FMaterialRenderProxy* UseSecondaryMaterialProxy = (SecondaryMaterialProxy != nullptr) ? SecondaryMaterialProxy : MaterialProxy;
+				FMaterialRenderProxy* UseSecondaryMaterialProxy = (UseSecondaryMaterial != nullptr) ? UseSecondaryMaterial->GetRenderProxy() : MaterialProxy;
 				if (bDrawSecondaryBuffers && BufferSet->SecondaryIndexBuffer.Indices.Num() > 0 && UseSecondaryMaterialProxy != nullptr)
 				{
 					DrawBatch(Collector, *BufferSet, BufferSet->SecondaryIndexBuffer, UseSecondaryMaterialProxy, false, DepthPriority, ViewIndex, DynamicPrimitiveUniformBuffer);
@@ -216,8 +249,21 @@ void FBaseDynamicMeshSceneProxy::GetDynamicRayTracingInstances(FRayTracingMateri
 	TArray<FMeshRenderBufferSet*> Buffers;
 	GetActiveRenderBufferSets(Buffers);
 
-	FMaterialRenderProxy* SecondaryMaterialProxy =
-		ParentBaseComponent->HasSecondaryRenderMaterial() ? ParentBaseComponent->GetSecondaryRenderMaterial()->GetRenderProxy() : nullptr;
+	// will use this material instead of any others below, if it becomes non-null
+	UMaterialInterface* ForceOverrideMaterial = nullptr;
+	const bool bVertexColor = ParentBaseComponent->ColorMode == EDynamicMeshComponentColorOverrideMode::VertexColors ||
+		ParentBaseComponent->ColorMode == EDynamicMeshComponentColorOverrideMode::Polygroups ||
+		ParentBaseComponent->ColorMode == EDynamicMeshComponentColorOverrideMode::Constant;
+	if (bVertexColor)
+	{
+		ForceOverrideMaterial = UBaseDynamicMeshComponent::GetDefaultVertexColorMaterial_RenderThread();
+	}
+
+	UMaterialInterface* UseSecondaryMaterial = ForceOverrideMaterial;
+	if (ParentBaseComponent->HasSecondaryRenderMaterial() && ForceOverrideMaterial == nullptr)
+	{
+		UseSecondaryMaterial = ParentBaseComponent->GetSecondaryRenderMaterial();
+	}
 	bool bDrawSecondaryBuffers = ParentBaseComponent->GetSecondaryBuffersVisibility();
 
 	bool bHasPrecomputedVolumetricLightmap;
@@ -239,6 +285,10 @@ void FBaseDynamicMeshSceneProxy::GetDynamicRayTracingInstances(FRayTracingMateri
 		{
 			UseMaterial = ParentBaseComponent->GetOverrideRenderMaterial(0);
 		}
+		if (ForceOverrideMaterial)
+		{
+			UseMaterial = ForceOverrideMaterial;
+		}
 		FMaterialRenderProxy* MaterialProxy = UseMaterial->GetRenderProxy();
 
 		if (BufferSet->TriangleCount == 0)
@@ -262,7 +312,7 @@ void FBaseDynamicMeshSceneProxy::GetDynamicRayTracingInstances(FRayTracingMateri
 		}
 
 		// draw secondary index buffer if we have it, falling back to base material if we don't have the Secondary material
-		FMaterialRenderProxy* UseSecondaryMaterialProxy = (SecondaryMaterialProxy != nullptr) ? SecondaryMaterialProxy : MaterialProxy;
+		FMaterialRenderProxy* UseSecondaryMaterialProxy = (UseSecondaryMaterial != nullptr) ? UseSecondaryMaterial->GetRenderProxy() : MaterialProxy;
 		if (bDrawSecondaryBuffers
 			&& BufferSet->SecondaryIndexBuffer.Indices.Num() > 0
 			&& UseSecondaryMaterialProxy != nullptr
diff --git a/Engine/Source/Runtime/GeometryFramework/Private/Components/DynamicMeshComponent.cpp b/Engine/Source/Runtime/GeometryFramework/Private/Components/DynamicMeshComponent.cpp
--- a/Engine/Source/Runtime/GeometryFramework/Private/Components/DynamicMeshComponent.cpp
+++ b/Engine/Source/Runtime/GeometryFramework/Private/Components/DynamicMeshComponent.cpp
@@ -11,6 +11,7 @@
 #include "DynamicMesh/DynamicMeshAttributeSet.h"
 #include "DynamicMesh/MeshNormals.h"
 #include "MeshDescriptionToDynamicMesh.h"
+#include "Util/ColorConstants.h"
 
 #include "Changes/MeshVertexChange.h"
 #include "Changes/MeshChange.h"
@@ -821,6 +822,11 @@ FPrimitiveSceneProxy* UDynamicMeshComponent::CreateSceneProxy()
 			NewProxy->bUsePerTriangleColor = true;
 			NewProxy->PerTriangleColorFunc = [this](const FDynamicMesh3* MeshIn, int TriangleID) { return GetTriangleColor(MeshIn, TriangleID); };
 		}
+		else if ( GetColorOverrideMode() == EDynamicMeshComponentColorOverrideMode::Polygroups )
+		{
+			NewProxy->bUsePerTriangleColor = true;
+			NewProxy->PerTriangleColorFunc = [this](const FDynamicMesh3* MeshIn, int TriangleID) { return GetGroupColor(MeshIn, TriangleID); };
+		}
 
 		if (SecondaryTriFilterFunc)
 		{
@@ -934,6 +940,12 @@ FColor UDynamicMeshComponent::GetTriangleColor(const FDynamicMesh3* MeshIn, int
 }
 
 
+FColor UDynamicMeshComponent::GetGroupColor(const FDynamicMesh3* Mesh, int TriangleID) const
+{
+	int32 GroupID = Mesh->HasTriangleGroups() ? Mesh->GetTriangleGroup(TriangleID) : 0;
+	return UE::Geometry::LinearColors::SelectFColor(GroupID);
+}
+
 
 FBoxSphereBounds UDynamicMeshComponent::CalcBounds(const FTransform& LocalToWorld) const
 {
diff --git a/Engine/Source/Runtime/GeometryFramework/Private/GeometryFrameworkModule.cpp b/Engine/Source/Runtime/GeometryFramework/Private/GeometryFrameworkModule.cpp
--- a/Engine/Source/Runtime/GeometryFramework/Private/GeometryFrameworkModule.cpp
+++ b/Engine/Source/Runtime/GeometryFramework/Private/GeometryFrameworkModule.cpp
@@ -1,18 +1,25 @@
 // Copyright Epic Games, Inc. All Rights Reserved.
 
 #include "GeometryFrameworkModule.h"
+#include "Components/BaseDynamicMeshComponent.h"
 
 #define LOCTEXT_NAMESPACE "FGeometryFrameworkModule"
 
 void FGeometryFrameworkModule::StartupModule()
 {
-	// This code will execute after your module is loaded into memory; the exact timing is specified in the .uplugin file per-module
+	FCoreDelegates::OnPostEngineInit.AddRaw(this, &FGeometryFrameworkModule::OnPostEngineInit);
+}
+
+void FGeometryFrameworkModule::OnPostEngineInit()
+{
+	// UBaseDynamicMeshComponent provides some global materials to all instances, rather than
+	// directly accessing (eg) GEngine pointers. Initialize those here. 
+	UBaseDynamicMeshComponent::InitializeDefaultMaterials();
 }
 
 void FGeometryFrameworkModule::ShutdownModule()
 {
-	// This function may be called during shutdown to clean up your module.  For modules that support dynamic reloading,
-	// we call this function before unloading the module.
+	FCoreDelegates::OnPostEngineInit.RemoveAll(this);
 }
 
 #undef LOCTEXT_NAMESPACE
diff --git a/Engine/Source/Runtime/GeometryFramework/Public/Components/BaseDynamicMeshComponent.h b/Engine/Source/Runtime/GeometryFramework/Public/Components/BaseDynamicMeshComponent.h
--- a/Engine/Source/Runtime/GeometryFramework/Public/Components/BaseDynamicMeshComponent.h
+++ b/Engine/Source/Runtime/GeometryFramework/Public/Components/BaseDynamicMeshComponent.h
@@ -55,6 +55,24 @@ enum class EDynamicMeshComponentTangentsMode : uint8
 };
 
 
+/**
+ * Color Override Modes
+ */
+UENUM()
+enum class EDynamicMeshComponentColorOverrideMode : uint8
+{
+	/** No Color Override enabled */
+	None,
+	/** Vertex Colors are displayed */
+	VertexColors,
+	/** Polygroup Colors are displayed */
+	Polygroups,
+	/** Constant Color is displayed */
+	Constant
+};
+
+
+
 /**
  * UBaseDynamicMeshComponent is a base interface for a UMeshComponent based on a UDynamicMesh.
  */
@@ -204,6 +222,76 @@ public:
 	UFUNCTION(BlueprintCallable, Category = "Dynamic Mesh Component")
 	virtual bool GetEnableWireframeRenderPass() const { return bExplicitShowWireframe; }
 
+	/**
+	 * Constant Color used when Override Color Mode is set to Constant
+	 */
+	UPROPERTY(EditAnywhere, BlueprintReadWrite, Category = "Dynamic Mesh Component|Rendering", meta = (DisplayName = "Wireframe Color"))
+	FLinearColor WireframeColor = FLinearColor(0, 0.5f, 1.f);
+
+
+	//===============================================================================================================
+	// Built-in Color Rendering Support. When enabled, Color mode will override any assigned Materials.
+	// VertexColor mode displays vertex colors, Polygroup mode displays mesh polygroups via vertex colors,
+	// and Constant mode uses ConstantColor as the vertex color. The class-wide DefaultVertexColorMaterial
+	// is used as the material that displays the vertex colors, and cannot be overridden per-instance
+	// (the OverrideRenderMaterial can be used to do that)
+public:
+	UPROPERTY(EditAnywhere, BlueprintReadWrite, Category = "Dynamic Mesh Component|Rendering", meta = (DisplayName = "Color Override") )
+	EDynamicMeshComponentColorOverrideMode ColorMode = EDynamicMeshComponentColorOverrideMode::None;
+
+	/**
+	 * Configure the active Color Override
+	 */
+	UFUNCTION(BlueprintCallable, Category = "Dynamic Mesh Component|Rendering")
+	virtual void SetColorOverrideMode(EDynamicMeshComponentColorOverrideMode NewMode);
+
+	/**
+	 * @return active Color Override mode
+	 */
+	UFUNCTION(BlueprintCallable, Category = "Dynamic Mesh Component|Rendering")
+	virtual EDynamicMeshComponentColorOverrideMode GetColorOverrideMode() const { return ColorMode; }
+
+	/**
+	 * Constant Color used when Override Color Mode is set to Constant
+	 */
+	UPROPERTY(EditAnywhere, BlueprintReadWrite, Category = "Dynamic Mesh Component|Rendering", meta = (DisplayName = "Constant Color", EditCondition = "ColorMode==EDynamicMeshComponentColorOverrideMode::Constant"))
+	FColor ConstantColor = FColor::White;
+
+	/**
+	 * Configure the Color used with Constant Color Override Mode
+	 */
+	UFUNCTION(BlueprintCallable, Category = "Dynamic Mesh Component|Rendering")
+	virtual void SetConstantOverrideColor(FColor NewColor);
+
+	/**
+	 * @return active Color used for Constant Color Override Mode
+	 */
+	UFUNCTION(BlueprintCallable, Category = "Dynamic Mesh Component|Rendering")
+	virtual FColor GetConstantOverrideColor() const { return ConstantColor; }
+
+
+	//===============================================================================================================
+	// Flat shading support. When enabled, per-triangle normals are computed automatically and used in place
+	// of the mesh normals. Mesh tangents are not affected.
+public:
+	/**
+	 * Enable use of per-triangle facet normals in place of mesh normals
+	 */
+	UPROPERTY(EditAnywhere, BlueprintReadWrite, Category = "Dynamic Mesh Component|Rendering", meta = (DisplayName = "Flat Shading") )
+	bool bEnableFlatShading = false;
+
+	/**
+	 * Configure the Color used with Constant Color Override Mode
+	 */
+	UFUNCTION(BlueprintCallable, Category = "Dynamic Mesh Component|Rendering")
+	virtual void SetEnableFlatShading(bool bEnable);
+
+	/**
+	 * @return active Color used for Constant Color Override Mode
+	 */
+	UFUNCTION(BlueprintCallable, Category = "Dynamic Mesh Component|Rendering")
+	virtual bool GetFlatShadingEnabled() const { return bEnableFlatShading; }
+
 
 
 	//===============================================================================================================
@@ -387,4 +475,40 @@ public:
 
 	UPROPERTY()
 	TArray<TObjectPtr<UMaterialInterface>> BaseMaterials;
+
+
+
+
+
+	//===============================================================================================================
+	// Class-wide Default Materials used for Wireframe and VertexColor display mode.
+	// These are configured for the Editor when the module loads, defaulting to built-in Engine wireframe and
+	// vertex color materials. 
+	// Note that the ModelingComponents module in the MeshModelingToolset plugin (usually enabled in the UE Editor) 
+	// will set a new VertexColor material from that plugins Content. 
+	// Client code can further configure these materials as necessary using the static functions below.
+public:
+	/**
+	 * Set the wireframe material used for all BaseDynamicMeshComponent-derived Components
+	 */
+	static void SetDefaultWireframeMaterial(UMaterialInterface* Material);
+
+	/**
+	 * Set the vertex color material used for all BaseDynamicMeshComponent-derived Components
+	 */
+	static void SetDefaultVertexColorMaterial(UMaterialInterface* Material);
+
+protected:
+	static void InitializeDefaultMaterials();
+	friend class FGeometryFrameworkModule;			// FGeometryFrameworkModule needs to call the above function
+
+	static UMaterialInterface* GetDefaultWireframeMaterial_RenderThread();
+	static UMaterialInterface* GetDefaultVertexColorMaterial_RenderThread();
+	friend class FBaseDynamicMeshSceneProxy;		// FBaseDynamicMeshSceneProxy needs to call these functions...
+
+private:
+	// these Materials are used by the render thread. Once the engine is running they should not be modified without 
+	// using SetDefaultWireframeMaterial/SetDefaultVertexColorMaterial
+	static UMaterialInterface* DefaultWireframeMaterial;
+	static UMaterialInterface* DefaultVertexColorMaterial;
 };
diff --git a/Engine/Source/Runtime/GeometryFramework/Public/Components/BaseDynamicMeshSceneProxy.h b/Engine/Source/Runtime/GeometryFramework/Public/Components/BaseDynamicMeshSceneProxy.h
--- a/Engine/Source/Runtime/GeometryFramework/Public/Components/BaseDynamicMeshSceneProxy.h
+++ b/Engine/Source/Runtime/GeometryFramework/Public/Components/BaseDynamicMeshSceneProxy.h
@@ -433,6 +433,10 @@ public:
 	 */
 	TFunction<FColor(const FDynamicMesh3*, int)> PerTriangleColorFunc = nullptr;
 
+	/**
+	* If true, a facet normals are used instead of mesh normals
+	*/
+	bool bUsePerTriangleNormals = false;
 
 	/**
 	 * If true, populate secondary buffers using SecondaryTriFilterFunc
@@ -583,8 +587,16 @@ public:
 			{
 				RenderBuffers->PositionVertexBuffer.VertexPosition(VertIdx) = (FVector3f)Mesh->GetVertex(Tri[j]);
 
-				FVector3f Normal = (NormalOverlay != nullptr && TriNormal[j] != FDynamicMesh3::InvalidID) ?
-					NormalOverlay->GetElement(TriNormal[j]) : Mesh->GetVertexNormal(Tri[j]);
+				FVector3f Normal;
+				if (bUsePerTriangleNormals)
+				{
+					Normal = (FVector3f)Mesh->GetTriNormal(TriangleID);
+				}
+				else
+				{
+					Normal = (NormalOverlay != nullptr && TriNormal[j] != FDynamicMesh3::InvalidID) ?
+						NormalOverlay->GetElement(TriNormal[j]) : Mesh->GetVertexNormal(Tri[j]);
+				}
 
 				// get tangents
 				TangentsFunc(Tri[j], TriangleID, j, Normal, TangentX, TangentY);
@@ -774,8 +786,17 @@ public:
 				if (bUpdateNormals)
 				{
 					// get normal and tangent
-					FVector3f Normal = (NormalOverlay != nullptr && TriNormal[j] != FDynamicMesh3::InvalidID) ?
-						NormalOverlay->GetElement(TriNormal[j]) : Mesh->GetVertexNormal(Tri[j]);
+					FVector3f Normal;
+					if (bUsePerTriangleNormals)
+					{
+						Normal = (FVector3f)Mesh->GetTriNormal(TriangleID);
+					}
+					else
+					{
+						Normal = (NormalOverlay != nullptr && TriNormal[j] != FDynamicMesh3::InvalidID) ?
+							NormalOverlay->GetElement(TriNormal[j]) : Mesh->GetVertexNormal(Tri[j]);
+					}
+
 					TangentsFunc(Tri[j], TriangleID, j, Normal, TangentX, TangentY);
 
 					RenderBuffers->StaticMeshVertexBuffer.SetVertexTangents(VertIdx, (FVector3f)TangentX, (FVector3f)TangentY, (FVector3f)Normal);
diff --git a/Engine/Source/Runtime/GeometryFramework/Public/Components/DynamicMeshComponent.h b/Engine/Source/Runtime/GeometryFramework/Public/Components/DynamicMeshComponent.h
--- a/Engine/Source/Runtime/GeometryFramework/Public/Components/DynamicMeshComponent.h
+++ b/Engine/Source/Runtime/GeometryFramework/Public/Components/DynamicMeshComponent.h
@@ -295,6 +295,8 @@ protected:
 	/** This function is passed via lambda to the RenderProxy to be able to access TriangleColorFunc */
 	FColor GetTriangleColor(const FDynamicMesh3* Mesh, int TriangleID);
 
+	/** This function is passed via lambda to the RenderProxy when BaseDynamicMeshComponent::ColorMode == Polygroups */
+	FColor GetGroupColor(const FDynamicMesh3* Mesh, int TriangleID) const;
 
 
 
diff --git a/Engine/Source/Runtime/GeometryFramework/Public/GeometryFrameworkModule.h b/Engine/Source/Runtime/GeometryFramework/Public/GeometryFrameworkModule.h
--- a/Engine/Source/Runtime/GeometryFramework/Public/GeometryFrameworkModule.h
+++ b/Engine/Source/Runtime/GeometryFramework/Public/GeometryFrameworkModule.h
@@ -12,4 +12,8 @@ public:
 	/** IModuleInterface implementation */
 	virtual void StartupModule() override;
 	virtual void ShutdownModule() override;
+
+
+	// registered in StartupModule() with FCoreDelegates::OnPostEngineInit and called once
+	void OnPostEngineInit();
 };
