GeometryFramework: add support for fast vertex-buffer-only updates in UDynamicMeshComponent::OnMeshObjectChanged(), if the posted change type is a vertex deformation.

diff --git a/Engine/Source/Runtime/GeometryFramework/Private/Components/DynamicMeshComponent.cpp b/Engine/Source/Runtime/GeometryFramework/Private/Components/DynamicMeshComponent.cpp
--- a/Engine/Source/Runtime/GeometryFramework/Private/Components/DynamicMeshComponent.cpp
+++ b/Engine/Source/Runtime/GeometryFramework/Private/Components/DynamicMeshComponent.cpp
@@ -34,6 +34,33 @@ namespace
 
 
 
+namespace UELocal
+{
+	static EMeshRenderAttributeFlags ConvertChangeFlagsToUpdateFlags(EDynamicMeshAttributeChangeFlags ChangeFlags)
+	{
+		EMeshRenderAttributeFlags UpdateFlags = EMeshRenderAttributeFlags::None;
+		if ((ChangeFlags & EDynamicMeshAttributeChangeFlags::VertexPositions) != EDynamicMeshAttributeChangeFlags::Unknown)
+		{
+			UpdateFlags |= EMeshRenderAttributeFlags::Positions;
+		}
+		if ((ChangeFlags & EDynamicMeshAttributeChangeFlags::NormalsTangents) != EDynamicMeshAttributeChangeFlags::Unknown)
+		{
+			UpdateFlags |= EMeshRenderAttributeFlags::VertexNormals;
+		}
+		if ((ChangeFlags & EDynamicMeshAttributeChangeFlags::VertexColors) != EDynamicMeshAttributeChangeFlags::Unknown)
+		{
+			UpdateFlags |= EMeshRenderAttributeFlags::VertexColors;
+		}
+		if ((ChangeFlags & EDynamicMeshAttributeChangeFlags::UVs) != EDynamicMeshAttributeChangeFlags::Unknown)
+		{
+			UpdateFlags |= EMeshRenderAttributeFlags::VertexUVs;
+		}
+		return UpdateFlags;
+	}
+
+}
+
+
 
 UDynamicMeshComponent::UDynamicMeshComponent(const FObjectInitializer& ObjectInitializer)
 	: Super(ObjectInitializer)
@@ -1007,7 +1034,17 @@ void UDynamicMeshComponent::OnMeshObjectChanged(UDynamicMesh* ChangedMeshObject,
 	}
 	else
 	{
-		NotifyMeshUpdated();
+		if (ChangeInfo.Type == EDynamicMeshChangeType::DeformationEdit)
+		{
+			// if ChangeType is a vertex deformation, we can do a fast-update of the vertex buffers
+			// without fully rebuilding the SceneProxy
+			EMeshRenderAttributeFlags UpdateFlags = UELocal::ConvertChangeFlagsToUpdateFlags(ChangeInfo.Flags);
+			FastNotifyVertexAttributesUpdated(UpdateFlags);
+		}
+		else
+		{
+			NotifyMeshUpdated();
+		}
 		OnMeshChanged.Broadcast();
 	}
 
