GeometryFramework: add to UDynamicMeshComponent::SetTransientDeferCollisionUpdates function and flag. This transient setting allows C++ code (eg like code that implements a live 3D transform of mesh vertices) to temporarily disable collision re-generation without modifying the serialized UProperty.

diff --git a/Engine/Source/Runtime/GeometryFramework/Private/Components/DynamicMeshComponent.cpp b/Engine/Source/Runtime/GeometryFramework/Private/Components/DynamicMeshComponent.cpp
--- a/Engine/Source/Runtime/GeometryFramework/Private/Components/DynamicMeshComponent.cpp
+++ b/Engine/Source/Runtime/GeometryFramework/Private/Components/DynamicMeshComponent.cpp
@@ -1050,7 +1050,7 @@ void UDynamicMeshComponent::OnMeshObjectChanged(UDynamicMesh* ChangedMeshObject,
 
 	// Rebuild body setup. Should this be deferred until proxy creation? Sometimes multiple changes are emitted...
 	// todo: can possibly skip this in some change situations, eg if only changing attributes
-	if (bDeferCollisionUpdates)
+	if (bDeferCollisionUpdates || bTransientDeferCollisionUpdates )
 	{
 		InvalidatePhysicsData();
 	}
@@ -1080,7 +1080,7 @@ void UDynamicMeshComponent::SetDynamicMesh(UDynamicMesh* NewMesh)
 	OnMeshChanged.Broadcast();
 
 	// Rebuild physics data
-	if (bDeferCollisionUpdates)
+	if (bDeferCollisionUpdates || bTransientDeferCollisionUpdates)
 	{
 		InvalidatePhysicsData();
 	}
@@ -1400,6 +1400,14 @@ void UDynamicMeshComponent::SetDeferredCollisionUpdatesEnabled(bool bEnabled, bo
 	}
 }
 
+void UDynamicMeshComponent::SetTransientDeferCollisionUpdates(bool bEnabled)
+{
+	if (bTransientDeferCollisionUpdates != bEnabled)
+	{
+		bTransientDeferCollisionUpdates = bEnabled;
+	}
+}
+
 void UDynamicMeshComponent::SetSceneProxyVerifyUsedMaterials(bool bState)
 {
 	bProxyVerifyUsedMaterials = bState;
diff --git a/Engine/Source/Runtime/GeometryFramework/Public/Components/DynamicMeshComponent.h b/Engine/Source/Runtime/GeometryFramework/Public/Components/DynamicMeshComponent.h
--- a/Engine/Source/Runtime/GeometryFramework/Public/Components/DynamicMeshComponent.h
+++ b/Engine/Source/Runtime/GeometryFramework/Public/Components/DynamicMeshComponent.h
@@ -537,6 +537,14 @@ public:
 	UPROPERTY(EditAnywhere, BlueprintReadWrite, Category = "Dynamic Mesh Component|Collision");
 	bool bDeferCollisionUpdates = false;
 
+	/**
+	 * The bDeferCollisionUpdates property is a serialized Component setting that controls when collision is regenerated.
+	 * However, in some cases where a UDynamicMesh is being programmatically edited (eg like a live deformation/etc), we
+	 * want to defer collision updates during the user interaction, but not change the serialized property. 
+	 * In this case a non-serialized version of the flag can be *temporarily* set via this function in C++. 
+	 */
+	void SetTransientDeferCollisionUpdates(bool bEnabled);
+
 protected:
 	UPROPERTY(Instanced)
 	TObjectPtr<UBodySetup> MeshBodySetup;
@@ -544,6 +552,7 @@ protected:
 	virtual void InvalidatePhysicsData();
 	virtual void RebuildPhysicsData();
 
+	bool bTransientDeferCollisionUpdates = false;
 	bool bCollisionUpdatePending = false;
 
 protected:
