handle negative scales in UDynamicMeshComponent::ApplyTransform

diff --git a/Engine/Source/Runtime/GeometryFramework/Private/Components/DynamicMeshComponent.cpp b/Engine/Source/Runtime/GeometryFramework/Private/Components/DynamicMeshComponent.cpp
--- a/Engine/Source/Runtime/GeometryFramework/Private/Components/DynamicMeshComponent.cpp
+++ b/Engine/Source/Runtime/GeometryFramework/Private/Components/DynamicMeshComponent.cpp
@@ -196,11 +196,11 @@ void UDynamicMeshComponent::ApplyTransform(const FTransform3d& Transform, bool b
 	{
 		if (bInvert)
 		{
-			MeshTransforms::ApplyTransformInverse(EditMesh, Transform);
+			MeshTransforms::ApplyTransformInverse(EditMesh, Transform, true);
 		}
 		else
 		{
-			MeshTransforms::ApplyTransform(EditMesh, Transform);
+			MeshTransforms::ApplyTransform(EditMesh, Transform, true);
 		}
 	}, EDynamicMeshChangeType::DeformationEdit);
 }
