GeometryFramework: filter out degenerate triangles when building complex collision mesh for DynamicMeshComponent

diff --git a/Engine/Source/Runtime/GeometryFramework/Private/Components/DynamicMeshComponent.cpp b/Engine/Source/Runtime/GeometryFramework/Private/Components/DynamicMeshComponent.cpp
--- a/Engine/Source/Runtime/GeometryFramework/Private/Components/DynamicMeshComponent.cpp
+++ b/Engine/Source/Runtime/GeometryFramework/Private/Components/DynamicMeshComponent.cpp
@@ -1149,6 +1149,23 @@ bool UDynamicMeshComponent::GetPhysicsTriMeshData(struct FTriMeshCollisionData*
 			Triangle.v0 = (bIsSparseV) ? VertexMap[Tri.A] : Tri.A;
 			Triangle.v1 = (bIsSparseV) ? VertexMap[Tri.B] : Tri.B;
 			Triangle.v2 = (bIsSparseV) ? VertexMap[Tri.C] : Tri.C;
+
+			// Filter out triangles which will cause physics system to emit degenerate-geometry warnings.
+			// These checks reproduce tests in Chaos::CleanTrimesh
+			const FVector3f& A = CollisionData->Vertices[Triangle.v0];
+			const FVector3f& B = CollisionData->Vertices[Triangle.v1];
+			const FVector3f& C = CollisionData->Vertices[Triangle.v2];
+			if (A == B || A == C || B == C)
+			{
+				continue;
+			}
+			// anything that fails the first check should also fail this, but Chaos does both so doing the same here...
+			const float SquaredArea = FVector3f::CrossProduct(A - B, A - C).SizeSquared();
+			if (SquaredArea < UE_SMALL_NUMBER)
+			{
+				continue;
+			}
+
 			CollisionData->Indices.Add(Triangle);
 
 			CollisionData->MaterialIndices.Add(0);		// not supporting physical materials yet
