GeometryFramework: make BaseDynamicMeshSceneProxy support Visible in Ray Tracing flag. However this does not override the Component-specific Enable Raytracing flag.

diff --git a/Engine/Source/Runtime/GeometryFramework/Private/Components/BaseDynamicMeshSceneProxy.cpp b/Engine/Source/Runtime/GeometryFramework/Private/Components/BaseDynamicMeshSceneProxy.cpp
--- a/Engine/Source/Runtime/GeometryFramework/Private/Components/BaseDynamicMeshSceneProxy.cpp
+++ b/Engine/Source/Runtime/GeometryFramework/Private/Components/BaseDynamicMeshSceneProxy.cpp
@@ -37,7 +37,7 @@ FMeshRenderBufferSet* FBaseDynamicMeshSceneProxy::AllocateNewRenderBufferSet()
 	FMeshRenderBufferSet* RenderBufferSet = new FMeshRenderBufferSet(GetScene().GetFeatureLevel());
 
 	RenderBufferSet->Material = UMaterial::GetDefaultMaterial(MD_Surface);
-	RenderBufferSet->bEnableRaytracing = this->bEnableRaytracing;
+	RenderBufferSet->bEnableRaytracing = this->bEnableRaytracing && this->IsVisibleInRayTracing();
 
 	AllocatedSetsLock.Lock();
 	AllocatedBufferSets.Add(RenderBufferSet);
@@ -240,6 +240,12 @@ bool FBaseDynamicMeshSceneProxy::IsRayTracingRelevant() const
 	return true;
 }
 
+bool FBaseDynamicMeshSceneProxy::HasRayTracingRepresentation() const
+{
+	return true;
+}
+
+
 void FBaseDynamicMeshSceneProxy::GetDynamicRayTracingInstances(FRayTracingMaterialGatheringContext& Context, TArray<FRayTracingInstance>& OutRayTracingInstances)
 {
 	QUICK_SCOPE_CYCLE_COUNTER(STAT_BaseDynamicMeshSceneProxy_GetDynamicRayTracingInstances);
diff --git a/Engine/Source/Runtime/GeometryFramework/Public/Components/BaseDynamicMeshSceneProxy.h b/Engine/Source/Runtime/GeometryFramework/Public/Components/BaseDynamicMeshSceneProxy.h
--- a/Engine/Source/Runtime/GeometryFramework/Public/Components/BaseDynamicMeshSceneProxy.h
+++ b/Engine/Source/Runtime/GeometryFramework/Public/Components/BaseDynamicMeshSceneProxy.h
@@ -928,7 +928,7 @@ public:
 #if RHI_RAYTRACING
 
 	virtual bool IsRayTracingRelevant() const override;
-
+	virtual bool HasRayTracingRepresentation() const override;
 
 	virtual void GetDynamicRayTracingInstances(FRayTracingMaterialGatheringContext& Context, TArray<FRayTracingInstance>& OutRayTracingInstances) override;
 
