GeometryFramework: resolve issue w/ UDynamicMeshComponent::SetDynamicMesh() where the incoming mesh may not be Outered to the Component, so it won't necessarily be serialized, leading to potential situations where the DynamicMesh is null on level load. Also updated UDynamicMeshComponent::PostLoad() to handle this case instead of just check()ing.

diff --git a/Engine/Source/Runtime/GeometryFramework/Private/Components/DynamicMeshComponent.cpp b/Engine/Source/Runtime/GeometryFramework/Private/Components/DynamicMeshComponent.cpp
--- a/Engine/Source/Runtime/GeometryFramework/Private/Components/DynamicMeshComponent.cpp
+++ b/Engine/Source/Runtime/GeometryFramework/Private/Components/DynamicMeshComponent.cpp
@@ -85,7 +85,14 @@ void UDynamicMeshComponent::PostLoad()
 {
 	Super::PostLoad();
 
-	check(MeshObject != nullptr);
+	// The intention here is that MeshObject is never nullptr, however we cannot guarantee this as a subclass
+	// may have set it to null, and/or some type of serialization issue has caused it to fail to save/load.
+	// Avoid immediate crashes by creating a new UDynamicMesh here in such cases
+	if (ensure(MeshObject != nullptr) == false)
+	{
+		MeshObject = NewObject<UDynamicMesh>(this, TEXT("DynamicMesh"));
+	}
+
 	MeshObjectChangedHandle = MeshObject->OnMeshChanged().AddUObject(this, &UDynamicMeshComponent::OnMeshObjectChanged);
 
 	ResetProxy();
@@ -1076,6 +1083,7 @@ void UDynamicMeshComponent::SetDynamicMesh(UDynamicMesh* NewMesh)
 		MeshObject->OnMeshChanged().Remove(MeshObjectChangedHandle);
 	}
 
+	NewMesh->Rename( nullptr, this );		// set Outer of NewMesh to be this Component
 	MeshObject = NewMesh;
 	MeshObjectChangedHandle = MeshObject->OnMeshChanged().AddUObject(this, &UDynamicMeshComponent::OnMeshObjectChanged);
 
