DynamicMesh: resolve issue w/ wireframe not rendering for DynamicMeshComponent in shipping builds

#localization none
#tests preflight and fix verified in shipping build locally

diff --git a/Engine/Source/Runtime/GeometryFramework/Private/Components/BaseDynamicMeshSceneProxy.cpp b/Engine/Source/Runtime/GeometryFramework/Private/Components/BaseDynamicMeshSceneProxy.cpp
--- a/Engine/Source/Runtime/GeometryFramework/Private/Components/BaseDynamicMeshSceneProxy.cpp
+++ b/Engine/Source/Runtime/GeometryFramework/Private/Components/BaseDynamicMeshSceneProxy.cpp
@@ -100,8 +100,9 @@ void FBaseDynamicMeshSceneProxy::GetDynamicMeshElements(const TArray<const FScen
 {
 	QUICK_SCOPE_CYCLE_COUNTER(STAT_BaseDynamicMeshSceneProxy_GetDynamicMeshElements);
 
-	bool bWireframe = (AllowDebugViewmodes() && ViewFamily.EngineShowFlags.Wireframe)
-		|| ParentBaseComponent->GetEnableWireframeRenderPass();
+	bool bIsWireframeViewMode = (AllowDebugViewmodes() && ViewFamily.EngineShowFlags.Wireframe);
+	bool bWantWireframeOnShaded = ParentBaseComponent->GetEnableWireframeRenderPass();
+	bool bWireframe = bIsWireframeViewMode || bWantWireframeOnShaded;
 
 	// Get wireframe material proxy if requested and available, otherwise disable wireframe
 	FMaterialRenderProxy* WireframeMaterialProxy = nullptr;
@@ -110,8 +111,10 @@ void FBaseDynamicMeshSceneProxy::GetDynamicMeshElements(const TArray<const FScen
 		UMaterialInterface* WireframeMaterial = UBaseDynamicMeshComponent::GetDefaultWireframeMaterial_RenderThread();
 		if (WireframeMaterial != nullptr)
 		{
+			FLinearColor UseWireframeColor = (IsSelected() && (bWantWireframeOnShaded == false || bIsWireframeViewMode))  ?
+				GEngine->GetSelectedMaterialColor() : ParentBaseComponent->WireframeColor;
 			FColoredMaterialRenderProxy* WireframeMaterialInstance = new FColoredMaterialRenderProxy(
-				WireframeMaterial->GetRenderProxy(), ParentBaseComponent->WireframeColor);
+				WireframeMaterial->GetRenderProxy(), UseWireframeColor);
 			Collector.RegisterOneFrameMaterialProxy(WireframeMaterialInstance);
 			WireframeMaterialProxy = WireframeMaterialInstance;
 		}
@@ -186,18 +189,23 @@ void FBaseDynamicMeshSceneProxy::GetDynamicMeshElements(const TArray<const FScen
 				FDynamicPrimitiveUniformBuffer& DynamicPrimitiveUniformBuffer = Collector.AllocateOneFrameResource<FDynamicPrimitiveUniformBuffer>();
 				DynamicPrimitiveUniformBuffer.Set(GetLocalToWorld(), PreviousLocalToWorld, GetBounds(), GetLocalBounds(), GetLocalBounds(), true, bHasPrecomputedVolumetricLightmap, bOutputVelocity, GetCustomPrimitiveData());
 
+				// If we want Wireframe-on-Shaded, we have to draw the solid. If View Mode Overrides are enabled, the solid
+				// will be replaced with it's wireframe, so we might as well not. 
+				bool bDrawSolidWithWireframe = ( bWantWireframeOnShaded && (bIsWireframeViewMode == false || bEnableViewModeOverrides == false) );
+
 				if (BufferSet->IndexBuffer.Indices.Num() > 0)
 				{
-					// Unlike most meshes, which just use the wireframe material in wireframe mode, we draw the wireframe on top of the normal material if needed,
-					// as this is easier to interpret. However, we do not do this in ortho viewports, where it frequently causes the our edit gizmo to be hidden 
-					// beneath the material. So, only draw the base material if we are in perspective mode, or we're in ortho but not in wireframe.
-					if (View->IsPerspectiveProjection() || !(AllowDebugViewmodes() && ViewFamily.EngineShowFlags.Wireframe))
+					if (bWireframe)
 					{
-						DrawBatch(Collector, *BufferSet, BufferSet->IndexBuffer, MaterialProxy, false, DepthPriority, ViewIndex, DynamicPrimitiveUniformBuffer);
+						if (bDrawSolidWithWireframe)
+						{
+							DrawBatch(Collector, *BufferSet, BufferSet->IndexBuffer, MaterialProxy, /*bWireframe*/false, DepthPriority, ViewIndex, DynamicPrimitiveUniformBuffer);
+						}
+						DrawBatch(Collector, *BufferSet, BufferSet->IndexBuffer, WireframeMaterialProxy, /*bWireframe*/true, DepthPriority, ViewIndex, DynamicPrimitiveUniformBuffer);
 					}
-					if (bWireframe)
+					else
 					{
-						DrawBatch(Collector, *BufferSet, BufferSet->IndexBuffer, WireframeMaterialProxy, true, DepthPriority, ViewIndex, DynamicPrimitiveUniformBuffer);
+						DrawBatch(Collector, *BufferSet, BufferSet->IndexBuffer, MaterialProxy, /*bWireframe*/false, DepthPriority, ViewIndex, DynamicPrimitiveUniformBuffer);
 					}
 				}
 
@@ -205,10 +213,17 @@ void FBaseDynamicMeshSceneProxy::GetDynamicMeshElements(const TArray<const FScen
 				FMaterialRenderProxy* UseSecondaryMaterialProxy = (UseSecondaryMaterial != nullptr) ? UseSecondaryMaterial->GetRenderProxy() : MaterialProxy;
 				if (bDrawSecondaryBuffers && BufferSet->SecondaryIndexBuffer.Indices.Num() > 0 && UseSecondaryMaterialProxy != nullptr)
 				{
-					DrawBatch(Collector, *BufferSet, BufferSet->SecondaryIndexBuffer, UseSecondaryMaterialProxy, false, DepthPriority, ViewIndex, DynamicPrimitiveUniformBuffer);
 					if (bWireframe)
 					{
-						DrawBatch(Collector, *BufferSet, BufferSet->SecondaryIndexBuffer, UseSecondaryMaterialProxy, true, DepthPriority, ViewIndex, DynamicPrimitiveUniformBuffer);
+						if (bDrawSolidWithWireframe)
+						{
+							DrawBatch(Collector, *BufferSet, BufferSet->SecondaryIndexBuffer, UseSecondaryMaterialProxy, /*bWireframe*/false, DepthPriority, ViewIndex, DynamicPrimitiveUniformBuffer);
+						}
+						DrawBatch(Collector, *BufferSet, BufferSet->SecondaryIndexBuffer, UseSecondaryMaterialProxy, /*bWireframe*/true, DepthPriority, ViewIndex, DynamicPrimitiveUniformBuffer);
+					}
+					else
+					{
+						DrawBatch(Collector, *BufferSet, BufferSet->SecondaryIndexBuffer, UseSecondaryMaterialProxy, /*bWireframe*/false, DepthPriority, ViewIndex, DynamicPrimitiveUniformBuffer);
 					}
 				}
 			}
@@ -222,6 +237,7 @@ void FBaseDynamicMeshSceneProxy::DrawBatch(FMeshElementCollector& Collector, con
 	FMeshBatchElement& BatchElement = Mesh.Elements[0];
 	BatchElement.IndexBuffer = &IndexBuffer;
 	Mesh.bWireframe = bWireframe;
+	//Mesh.bDisableBackfaceCulling = bWireframe;		// todo: doing this would be more consistent w/ other meshes in wireframe mode, but it is problematic for modeling tools - perhaps should be configurable
 	Mesh.VertexFactory = &RenderBuffers.VertexFactory;
 	Mesh.MaterialRenderProxy = UseMaterial;
 
@@ -234,7 +250,8 @@ void FBaseDynamicMeshSceneProxy::DrawBatch(FMeshElementCollector& Collector, con
 	Mesh.ReverseCulling = IsLocalToWorldDeterminantNegative();
 	Mesh.Type = PT_TriangleList;
 	Mesh.DepthPriorityGroup = DepthPriority;
-	Mesh.bCanApplyViewModeOverrides = this->bEnableViewModeOverrides;
+	// if this is a wireframe draw pass then we do not want to apply View Mode Overrides
+	Mesh.bCanApplyViewModeOverrides = (bWireframe) ? false : this->bEnableViewModeOverrides;
 	Collector.AddMesh(ViewIndex, Mesh);
 }
 
