fix nonunity build error

diff --git a/Engine/Source/Runtime/GeometryFramework/Private/Components/BaseDynamicMeshSceneProxy.cpp b/Engine/Source/Runtime/GeometryFramework/Private/Components/BaseDynamicMeshSceneProxy.cpp
--- a/Engine/Source/Runtime/GeometryFramework/Private/Components/BaseDynamicMeshSceneProxy.cpp
+++ b/Engine/Source/Runtime/GeometryFramework/Private/Components/BaseDynamicMeshSceneProxy.cpp
@@ -8,6 +8,7 @@
 #include "RayTracingInstance.h"
 #include "SceneInterface.h"
 #include "SceneManagement.h"
+#include "Engine/Engine.h"		// for GEngine definition
 
 FBaseDynamicMeshSceneProxy::FBaseDynamicMeshSceneProxy(UBaseDynamicMeshComponent* Component)
 	: FPrimitiveSceneProxy(Component),
