GeometyFramework: Add support for controlling vertex color space transform/conversion in BaseDynamicMeshComponent. Add support for custom color remapping function in DynamicMeshComponent.

diff --git a/Engine/Source/Runtime/GeometryFramework/Private/Components/BaseDynamicMeshComponent.cpp b/Engine/Source/Runtime/GeometryFramework/Private/Components/BaseDynamicMeshComponent.cpp
--- a/Engine/Source/Runtime/GeometryFramework/Private/Components/BaseDynamicMeshComponent.cpp
+++ b/Engine/Source/Runtime/GeometryFramework/Private/Components/BaseDynamicMeshComponent.cpp
@@ -141,6 +141,19 @@ void UBaseDynamicMeshComponent::SetConstantOverrideColor(FColor NewColor)
 }
 
 
+void UBaseDynamicMeshComponent::SetVertexColorSpaceTransformMode(EDynamicMeshVertexColorTransformMode NewMode)
+{
+	if (ColorSpaceMode != NewMode)
+	{
+		ColorSpaceMode = NewMode;
+		OnRenderingStateChanged(false);
+	}
+}
+
+
+
+
+
 void UBaseDynamicMeshComponent::SetEnableFlatShading(bool bEnable)
 {
 	if (bEnableFlatShading != bEnable)
diff --git a/Engine/Source/Runtime/GeometryFramework/Private/Components/BaseDynamicMeshSceneProxy.cpp b/Engine/Source/Runtime/GeometryFramework/Private/Components/BaseDynamicMeshSceneProxy.cpp
--- a/Engine/Source/Runtime/GeometryFramework/Private/Components/BaseDynamicMeshSceneProxy.cpp
+++ b/Engine/Source/Runtime/GeometryFramework/Private/Components/BaseDynamicMeshSceneProxy.cpp
@@ -13,6 +13,7 @@
 FBaseDynamicMeshSceneProxy::FBaseDynamicMeshSceneProxy(UBaseDynamicMeshComponent* Component)
 	: FPrimitiveSceneProxy(Component),
 	ParentBaseComponent(Component),
+	ColorSpaceTransformMode(Component->GetVertexColorSpaceTransformMode()),
 	bEnableRaytracing(Component->GetEnableRaytracing()),
 	bEnableViewModeOverrides(Component->GetViewModeOverridesEnabled())
 {
diff --git a/Engine/Source/Runtime/GeometryFramework/Private/Components/DynamicMeshComponent.cpp b/Engine/Source/Runtime/GeometryFramework/Private/Components/DynamicMeshComponent.cpp
--- a/Engine/Source/Runtime/GeometryFramework/Private/Components/DynamicMeshComponent.cpp
+++ b/Engine/Source/Runtime/GeometryFramework/Private/Components/DynamicMeshComponent.cpp
@@ -384,6 +384,17 @@ void UDynamicMeshComponent::FastNotifyColorsUpdated()
 			Proxy->PerTriangleColorFunc = nullptr;
 		}
 
+		if (HasVertexColorRemappingFunction() && Proxy->bApplyVertexColorRemapping == false)
+		{
+			Proxy->bApplyVertexColorRemapping = true;
+			Proxy->VertexColorRemappingFunc = [this](FVector4f& Color) { RemapVertexColor(Color); };
+		}
+		else if (!HasVertexColorRemappingFunction() && Proxy->bApplyVertexColorRemapping == true)
+		{
+			Proxy->bApplyVertexColorRemapping = false;
+			Proxy->VertexColorRemappingFunc = nullptr;
+		}
+
 		Proxy->FastUpdateVertices(false, false, true, false);
 		//MarkRenderDynamicDataDirty();
 	}
@@ -865,6 +876,12 @@ FPrimitiveSceneProxy* UDynamicMeshComponent::CreateSceneProxy()
 			NewProxy->PerTriangleColorFunc = [this](const FDynamicMesh3* MeshIn, int TriangleID) { return GetGroupColor(MeshIn, TriangleID); };
 		}
 
+		if (HasVertexColorRemappingFunction())
+		{
+			NewProxy->bApplyVertexColorRemapping = true;
+			NewProxy->VertexColorRemappingFunc = [this](FVector4f& Color) { RemapVertexColor(Color); };
+		}
+
 		if (SecondaryTriFilterFunc)
 		{
 			NewProxy->bUseSecondaryTriBuffers = true;
@@ -941,6 +958,54 @@ bool UDynamicMeshComponent::HasTriangleColorFunction()
 
 
 
+void UDynamicMeshComponent::SetVertexColorRemappingFunction(
+	TUniqueFunction<void(FVector4f&)> ColorMapFuncIn,
+	EDynamicMeshComponentRenderUpdateMode UpdateMode)
+{
+	VertexColorMappingFunc = MoveTemp(ColorMapFuncIn);
+
+	if (UpdateMode == EDynamicMeshComponentRenderUpdateMode::FastUpdate)
+	{
+		FastNotifyColorsUpdated();
+	}
+	else if (UpdateMode == EDynamicMeshComponentRenderUpdateMode::FullUpdate)
+	{
+		NotifyMeshUpdated();
+	}
+}
+
+void UDynamicMeshComponent::ClearVertexColorRemappingFunction(EDynamicMeshComponentRenderUpdateMode UpdateMode)
+{
+	if (VertexColorMappingFunc)
+	{
+		VertexColorMappingFunc = nullptr;
+
+		if (UpdateMode == EDynamicMeshComponentRenderUpdateMode::FastUpdate)
+		{
+			FastNotifyColorsUpdated();
+		}
+		else if (UpdateMode == EDynamicMeshComponentRenderUpdateMode::FullUpdate)
+		{
+			NotifyMeshUpdated();
+		}
+	}
+}
+
+bool UDynamicMeshComponent::HasVertexColorRemappingFunction()
+{
+	return !!VertexColorMappingFunc;
+}
+
+
+void UDynamicMeshComponent::RemapVertexColor(FVector4f& VertexColorInOut)
+{
+	if (VertexColorMappingFunc)
+	{
+		VertexColorMappingFunc(VertexColorInOut);
+	}
+}
+
+
 
 void UDynamicMeshComponent::EnableSecondaryTriangleBuffers(TUniqueFunction<bool(const FDynamicMesh3*, int32)> SecondaryTriFilterFuncIn)
 {
diff --git a/Engine/Source/Runtime/GeometryFramework/Public/Components/BaseDynamicMeshComponent.h b/Engine/Source/Runtime/GeometryFramework/Public/Components/BaseDynamicMeshComponent.h
--- a/Engine/Source/Runtime/GeometryFramework/Public/Components/BaseDynamicMeshComponent.h
+++ b/Engine/Source/Runtime/GeometryFramework/Public/Components/BaseDynamicMeshComponent.h
@@ -72,6 +72,23 @@ enum class EDynamicMeshComponentColorOverrideMode : uint8
 };
 
 
+/**
+ * Color Transform to apply to Vertex Colors when converting from internal DynamicMesh
+ * Color attributes (eg Color Overlay stored in FVector4f) to RHI Render Buffers (FColor).
+ * 
+ * Note that UStaticMesh assumes the Source Mesh colors are Linear and always converts to SRGB.
+ */
+UENUM()
+enum class EDynamicMeshVertexColorTransformMode : uint8
+{
+	/** Do not apply any color-space transform to Vertex Colors */
+	NoTransform,
+	/** Assume Vertex Colors are in Linear space and transform to SRGB */
+	LinearToSRGB,
+	/** Assume Vertex Colors are in SRGB space and convert to Linear */
+	SRGBToLinear
+};
+
 
 /**
  * UBaseDynamicMeshComponent is a base interface for a UMeshComponent based on a UDynamicMesh.
@@ -269,6 +286,30 @@ public:
 	UFUNCTION(BlueprintCallable, Category = "Dynamic Mesh Component|Rendering")
 	virtual FColor GetConstantOverrideColor() const { return ConstantColor; }
 
+	/**
+	 * Color Space Transform that will be applied to the colors stored in the DynamicMesh Attribute Color Overlay when
+	 * constructing render buffers. 
+	 * Default is "No Transform", ie color R/G/B/A will be independently converted from 32-bit float to 8-bit by direct mapping.
+	 * LinearToSRGB mode will apply SRGB conversion, ie assumes colors in the Mesh are in Linear space. This will produce the same behavior as UStaticMesh.
+	 * SRGBToLinear mode will invert SRGB conversion, ie assumes colors in the Mesh are in SRGB space. 
+	 */
+	UPROPERTY(EditAnywhere, BlueprintReadWrite, Category = "Dynamic Mesh Component|Rendering", meta = (DisplayName = "Vertex Color Space"))
+	EDynamicMeshVertexColorTransformMode ColorSpaceMode = EDynamicMeshVertexColorTransformMode::NoTransform;
+	
+	/**
+	 * Configure the active Color Space Transform Mode
+	 */
+	UFUNCTION(BlueprintCallable, Category = "Dynamic Mesh Component|Rendering")
+	virtual void SetVertexColorSpaceTransformMode(EDynamicMeshVertexColorTransformMode NewMode);
+
+	/**
+	 * @return active Color Override mode
+	 */
+	UFUNCTION(BlueprintCallable, Category = "Dynamic Mesh Component|Rendering")
+	virtual EDynamicMeshVertexColorTransformMode GetVertexColorSpaceTransformMode() const { return ColorSpaceMode; }
+
+
+
 
 	//===============================================================================================================
 	// Flat shading support. When enabled, per-triangle normals are computed automatically and used in place
diff --git a/Engine/Source/Runtime/GeometryFramework/Public/Components/BaseDynamicMeshSceneProxy.h b/Engine/Source/Runtime/GeometryFramework/Public/Components/BaseDynamicMeshSceneProxy.h
--- a/Engine/Source/Runtime/GeometryFramework/Public/Components/BaseDynamicMeshSceneProxy.h
+++ b/Engine/Source/Runtime/GeometryFramework/Public/Components/BaseDynamicMeshSceneProxy.h
@@ -435,6 +435,23 @@ public:
 	 */
 	TFunction<FColor(const FDynamicMesh3*, int)> PerTriangleColorFunc = nullptr;
 
+
+	/**
+	 * If true, VertexColorRemappingFunc is called on Vertex Colors provided from Mesh to remap them to a different color
+	 */
+	bool bApplyVertexColorRemapping = false;
+
+	/**
+	 * Vertex color remapping function. Only called if bApplyVertexColorRemapping == true, for mesh vertex colors
+	 */
+	TUniqueFunction<void(FVector4f&)> VertexColorRemappingFunc = nullptr;
+
+	/**
+	 * Color Space Transform/Conversion applied to Vertex Colors provided from Mesh Color Overlay Attribute
+	 * Color Space Conversion is applied after any Vertex Color Remapping.
+	 */
+	EDynamicMeshVertexColorTransformMode ColorSpaceTransformMode = EDynamicMeshVertexColorTransformMode::NoTransform;
+
 	/**
 	* If true, a facet normals are used instead of mesh normals
 	*/
@@ -613,7 +630,7 @@ public:
 				}
 
 				FColor VertexFColor = (bHaveColors && TriColor[j] != FDynamicMesh3::InvalidID) ?
-					UE::Geometry::ToFColor(ColorOverlay->GetElement(TriColor[j])) : UniformTriColor;
+					GetOverlayColorAsFColor(ColorOverlay, TriColor[j]) : UniformTriColor;
 
 				RenderBuffers->ColorVertexBuffer.VertexColor(VertIdx) = VertexFColor;
 
@@ -636,6 +653,30 @@ public:
 	}
 
 
+	FColor GetOverlayColorAsFColor(
+		const FDynamicMeshColorOverlay* ColorOverlay,
+		int32 ElementID)
+	{
+		checkSlow(ColorOverlay);
+		FVector4f UseColor = ColorOverlay->GetElement(ElementID);
+
+		if (bApplyVertexColorRemapping)
+		{
+			VertexColorRemappingFunc(UseColor);
+		}
+
+		if (ColorSpaceTransformMode == EDynamicMeshVertexColorTransformMode::SRGBToLinear)
+		{
+			// is there a better way to do this? 
+			FColor QuantizedSRGBColor = ((FLinearColor)UseColor).ToFColor(false);
+			return FLinearColor(QuantizedSRGBColor).ToFColor(false);
+		}
+		else
+		{
+			bool bConvertToSRGB = (ColorSpaceTransformMode == EDynamicMeshVertexColorTransformMode::LinearToSRGB);
+			return ((FLinearColor)UseColor).ToFColor(bConvertToSRGB);
+		}
+	}
 
 
 	/**
@@ -807,7 +848,7 @@ public:
 				if (bUpdateColors)
 				{
 					FColor VertexFColor = (bHaveColors  && TriColor[j] != FDynamicMesh3::InvalidID) ?
-						UE::Geometry::ToFColor(ColorOverlay->GetElement(TriColor[j])) : UniformTriColor;
+						GetOverlayColorAsFColor(ColorOverlay, TriColor[j]) : UniformTriColor;
 					RenderBuffers->ColorVertexBuffer.VertexColor(VertIdx) = VertexFColor;
 				}
 
diff --git a/Engine/Source/Runtime/GeometryFramework/Public/Components/DynamicMeshComponent.h b/Engine/Source/Runtime/GeometryFramework/Public/Components/DynamicMeshComponent.h
--- a/Engine/Source/Runtime/GeometryFramework/Public/Components/DynamicMeshComponent.h
+++ b/Engine/Source/Runtime/GeometryFramework/Public/Components/DynamicMeshComponent.h
@@ -278,7 +278,7 @@ protected:
 	// minimal vertex buffer updates necessary in the RenderProxy
 	//
 public:
-	/** Clear an active triangle color function if one exists, and update the mesh */
+	/** Set an active triangle color function if one exists, and update the mesh */
 	virtual void SetTriangleColorFunction(TUniqueFunction<FColor(const FDynamicMesh3*, int)> TriangleColorFuncIn,
 										  EDynamicMeshComponentRenderUpdateMode UpdateMode = EDynamicMeshComponentRenderUpdateMode::FastUpdate);
 
@@ -299,6 +299,30 @@ protected:
 	FColor GetGroupColor(const FDynamicMesh3* Mesh, int TriangleID) const;
 
 
+	//===============================================================================================================
+	// Support for Vertex Color remapping/filtering. This allows external code to modulate the existing 
+	// Vertex Colors on the rendered mesh. The remapping is only applied to FVector4f Color Overlay attribute buffers.
+	// The lambda that is passed is held for the lifetime of the Component and
+	// must remain valid. If the Vertex Colors are modified, FastNotifyColorsUpdated() can be used to do the 
+	// minimal vertex buffer updates necessary in the RenderProxy
+public:
+	/** Set an active VertexColor Remapping function if one exists, and update the mesh */
+	virtual void SetVertexColorRemappingFunction(TUniqueFunction<void(FVector4f&)> ColorMapFuncIn,
+		EDynamicMeshComponentRenderUpdateMode UpdateMode = EDynamicMeshComponentRenderUpdateMode::FastUpdate);
+
+	/** Clear an active VertexColor Remapping function if one exists, and update the mesh */
+	virtual void ClearVertexColorRemappingFunction(EDynamicMeshComponentRenderUpdateMode UpdateMode = EDynamicMeshComponentRenderUpdateMode::FastUpdate);
+
+	/** @return true if a VertexColor Remapping function is configured */
+	virtual bool HasVertexColorRemappingFunction();
+
+protected:
+	/** If this function is set, DynamicMesh Attribute Color Overlay colors will be passed through this function before sending to render buffers */
+	TUniqueFunction<void(FVector4f&)> VertexColorMappingFunc = nullptr;
+
+	/** This function is passed via lambda to the RenderProxy to be able to access VertexColorMappingFunc */
+	void RemapVertexColor(FVector4f& VertexColorInOut);
+
 
 
 	//===============================================================================================================
