[Engine]

diff --git a/Engine/Source/Runtime/GeometryFramework/Public/Components/DynamicMeshComponent.h b/Engine/Source/Runtime/GeometryFramework/Public/Components/DynamicMeshComponent.h
--- a/Engine/Source/Runtime/GeometryFramework/Public/Components/DynamicMeshComponent.h
+++ b/Engine/Source/Runtime/GeometryFramework/Public/Components/DynamicMeshComponent.h
@@ -26,7 +26,7 @@ class FBaseDynamicMeshSceneProxy;
  * sending it off for rendering.
  * NOTE: This is called whenever the Mesh is updated and before rendering, so performance matters.
  */
-class GEOMETRYFRAMEWORK_API IRenderMeshPostProcessor
+class IRenderMeshPostProcessor
 {
 public:
 	virtual ~IRenderMeshPostProcessor() = default;
@@ -59,8 +59,8 @@ enum class EDynamicMeshComponentRenderUpdateMode
  * See comment sections below for details.
  * 
  */
-UCLASS(hidecategories = (LOD), meta = (BlueprintSpawnableComponent), ClassGroup = Rendering)
-class GEOMETRYFRAMEWORK_API UDynamicMeshComponent : public UBaseDynamicMeshComponent, public IInterface_CollisionDataProvider
+UCLASS(hidecategories = (LOD), meta = (BlueprintSpawnableComponent), ClassGroup = Rendering, MinimalAPI)
+class UDynamicMeshComponent : public UBaseDynamicMeshComponent, public IInterface_CollisionDataProvider
 {
 	GENERATED_UCLASS_BODY()
 
@@ -94,29 +94,29 @@ public:
 	 * @warning If NewMesh is owned/Outer'd to another DynamicMeshComponent, a GLEO error may occur if that Component is serialized.
 	 */
 	UFUNCTION(BlueprintCallable, Category = "Dynamic Mesh Component")
-	void SetDynamicMesh(UDynamicMesh* NewMesh);
+	GEOMETRYFRAMEWORK_API void SetDynamicMesh(UDynamicMesh* NewMesh);
 
 	/**
 	 * initialize the internal mesh from a DynamicMesh
 	 */
-	virtual void SetMesh(UE::Geometry::FDynamicMesh3&& MoveMesh) override;
+	GEOMETRYFRAMEWORK_API virtual void SetMesh(UE::Geometry::FDynamicMesh3&& MoveMesh) override;
 
 	/**
 	 * Allow external code to read the internal mesh.
 	 */
-	virtual void ProcessMesh(TFunctionRef<void(const UE::Geometry::FDynamicMesh3&)> ProcessFunc) const;
+	GEOMETRYFRAMEWORK_API virtual void ProcessMesh(TFunctionRef<void(const UE::Geometry::FDynamicMesh3&)> ProcessFunc) const;
 
 	/**
 	 * Allow external code to to edit the internal mesh.
 	 */
-	virtual void EditMesh(TFunctionRef<void(UE::Geometry::FDynamicMesh3&)> EditFunc,
+	GEOMETRYFRAMEWORK_API virtual void EditMesh(TFunctionRef<void(UE::Geometry::FDynamicMesh3&)> EditFunc,
 						  EDynamicMeshComponentRenderUpdateMode UpdateMode = EDynamicMeshComponentRenderUpdateMode::FullUpdate);
 
 	/**
 	 * Apply transform to internal mesh. In some cases this can be more efficient than a general edit.
 	 * @param bInvert if true, inverse tranform is applied instead of forward transform
 	 */
-	virtual void ApplyTransform(const FTransform3d& Transform, bool bInvert) override;
+	GEOMETRYFRAMEWORK_API virtual void ApplyTransform(const FTransform3d& Transform, bool bInvert) override;
 
 
 
@@ -142,55 +142,55 @@ public:
 	 * Call this if you update the mesh via GetMesh(). This will destroy the existing RenderProxy and create a new one.
 	 * @todo should provide a function that calls a lambda to modify the mesh, and only return const mesh pointer
 	 */
-	virtual void NotifyMeshUpdated() override;
+	GEOMETRYFRAMEWORK_API virtual void NotifyMeshUpdated() override;
 
 	/**
 	 * Call this instead of NotifyMeshUpdated() if you have only updated the vertex colors (or triangle color function).
 	 * This function will update the existing RenderProxy buffers if possible
 	 */
-	void FastNotifyColorsUpdated();
+	GEOMETRYFRAMEWORK_API void FastNotifyColorsUpdated();
 
 	/**
 	 * Call this instead of NotifyMeshUpdated() if you have only updated the vertex positions (and possibly some attributes).
 	 * This function will update the existing RenderProxy buffers if possible
 	 */
-	void FastNotifyPositionsUpdated(bool bNormals = false, bool bColors = false, bool bUVs = false);
+	GEOMETRYFRAMEWORK_API void FastNotifyPositionsUpdated(bool bNormals = false, bool bColors = false, bool bUVs = false);
 
 	/**
 	 * Call this instead of NotifyMeshUpdated() if you have only updated the vertex attributes (but not positions).
 	 * This function will update the existing RenderProxy buffers if possible, rather than create new ones.
 	 */
-	void FastNotifyVertexAttributesUpdated(bool bNormals, bool bColors, bool bUVs);
+	GEOMETRYFRAMEWORK_API void FastNotifyVertexAttributesUpdated(bool bNormals, bool bColors, bool bUVs);
 
 	/**
 	 * Call this instead of NotifyMeshUpdated() if you have only updated the vertex positions/attributes
 	 * This function will update the existing RenderProxy buffers if possible, rather than create new ones.
 	 */
-	void FastNotifyVertexAttributesUpdated(EMeshRenderAttributeFlags UpdatedAttributes);
+	GEOMETRYFRAMEWORK_API void FastNotifyVertexAttributesUpdated(EMeshRenderAttributeFlags UpdatedAttributes);
 
 	/**
 	 * Call this instead of NotifyMeshUpdated() if you have only updated the vertex uvs.
 	 * This function will update the existing RenderProxy buffers if possible
 	 */
-	void FastNotifyUVsUpdated();
+	GEOMETRYFRAMEWORK_API void FastNotifyUVsUpdated();
 
 	/**
 	 * Call this instead of NotifyMeshUpdated() if you have only updated secondary triangle sorting.
 	 * This function will update the existing buffers if possible, without rebuilding entire RenderProxy.
 	 */
-	void FastNotifySecondaryTrianglesChanged();
+	GEOMETRYFRAMEWORK_API void FastNotifySecondaryTrianglesChanged();
 
 	/**
 	 * This function updates vertex positions/attributes of existing SceneProxy render buffers if possible, for the given triangles.
 	 * If a FMeshRenderDecomposition has not been explicitly set, call is forwarded to FastNotifyVertexAttributesUpdated()
 	 */
-	void FastNotifyTriangleVerticesUpdated(const TArray<int32>& Triangles, EMeshRenderAttributeFlags UpdatedAttributes);
+	GEOMETRYFRAMEWORK_API void FastNotifyTriangleVerticesUpdated(const TArray<int32>& Triangles, EMeshRenderAttributeFlags UpdatedAttributes);
 
 	/**
 	 * This function updates vertex positions/attributes of existing SceneProxy render buffers if possible, for the given triangles.
 	 * If a FMeshRenderDecomposition has not been explicitly set, call is forwarded to FastNotifyVertexAttributesUpdated()
 	 */
-	void FastNotifyTriangleVerticesUpdated(const TSet<int32>& Triangles, EMeshRenderAttributeFlags UpdatedAttributes);
+	GEOMETRYFRAMEWORK_API void FastNotifyTriangleVerticesUpdated(const TSet<int32>& Triangles, EMeshRenderAttributeFlags UpdatedAttributes);
 
 	/**
 	 * If a Decomposition is set on this Component, and everything is currently valid (proxy/etc), precompute the set of
@@ -198,7 +198,7 @@ public:
 	 * Use FastNotifyTriangleVerticesUpdated_ApplyPrecompute() with the returned future to apply this precomputation.
 	 * @return a future that will (eventually) return true if the precompute is OK, and (immediately) false if it is not
 	 */
-	TFuture<bool> FastNotifyTriangleVerticesUpdated_TryPrecompute(const TArray<int32>& Triangles, TArray<int32>& UpdateSetsOut, UE::Geometry::FAxisAlignedBox3d& BoundsOut);
+	GEOMETRYFRAMEWORK_API TFuture<bool> FastNotifyTriangleVerticesUpdated_TryPrecompute(const TArray<int32>& Triangles, TArray<int32>& UpdateSetsOut, UE::Geometry::FAxisAlignedBox3d& BoundsOut);
 
 	/**
 	 * This function updates vertex positions/attributes of existing SceneProxy render buffers if possible, for the given triangles.
@@ -207,7 +207,7 @@ public:
 	 * passed to FastNotifyTriangleVerticesUpdated_TryPrecompute). 
 	 * If the Precompute future returns false, then we forward the call to FastNotifyTriangleVerticesUpdated(), which will do more work.
 	 */
-	void FastNotifyTriangleVerticesUpdated_ApplyPrecompute(const TArray<int32>& Triangles, EMeshRenderAttributeFlags UpdatedAttributes,
+	GEOMETRYFRAMEWORK_API void FastNotifyTriangleVerticesUpdated_ApplyPrecompute(const TArray<int32>& Triangles, EMeshRenderAttributeFlags UpdatedAttributes,
 		TFuture<bool>& Precompute, const TArray<int32>& UpdateSets, const UE::Geometry::FAxisAlignedBox3d& UpdateSetBounds);
 
 
@@ -225,7 +225,7 @@ public:
 	 * Notify Vertex Attributes Updated can be used to do a faster update.
 	 */
 	UFUNCTION(BlueprintCallable, Category = "Dynamic Mesh Component|Rendering", DisplayName = "Notify Mesh Updated")
-	virtual void NotifyMeshModified();
+	GEOMETRYFRAMEWORK_API virtual void NotifyMeshModified();
 
 	/**
 	 * Notify the Component that vertex attribute values of it's DynamicMesh have been modified externally. This will result in
@@ -236,7 +236,7 @@ public:
 	 * Behavior of this function is undefined in these cases and may crash. If you are unsure, use Notify Mesh Updated.
 	 */
 	UFUNCTION(BlueprintCallable, Category = "Dynamic Mesh Component|Rendering", DisplayName = "Notify Vertex Attributes Updated")
-	virtual void NotifyMeshVertexAttributesModified(
+	GEOMETRYFRAMEWORK_API virtual void NotifyMeshVertexAttributesModified(
 		bool bPositions = true, 
 		bool bNormals = true,
 		bool bUVs = true,
@@ -252,17 +252,17 @@ public:
 	/**
 	 * Apply a vertex deformation change to the mesh
 	 */
-	virtual void ApplyChange(const FMeshVertexChange* Change, bool bRevert) override;
+	GEOMETRYFRAMEWORK_API virtual void ApplyChange(const FMeshVertexChange* Change, bool bRevert) override;
 
 	/**
 	 * Apply a general mesh change to the mesh
 	 */
-	virtual void ApplyChange(const FMeshChange* Change, bool bRevert) override;
+	GEOMETRYFRAMEWORK_API virtual void ApplyChange(const FMeshChange* Change, bool bRevert) override;
 
 	/**
 	* Apply a mesh replacement change to mesh
 	*/
-	virtual void ApplyChange(const FMeshReplacementChange* Change, bool bRevert) override;
+	GEOMETRYFRAMEWORK_API virtual void ApplyChange(const FMeshReplacementChange* Change, bool bRevert) override;
 
 	/**
 	 * This delegate fires when the mesh has been changed
@@ -281,7 +281,7 @@ public:
 	 * (eg using one of the FastXYZ functions above). To allow for that, the full proxy invalidation on change can be (temporarily!) disabled 
 	 * using this function. 
 	 */
-	void SetInvalidateProxyOnChangeEnabled(bool bEnabled);
+	GEOMETRYFRAMEWORK_API void SetInvalidateProxyOnChangeEnabled(bool bEnabled);
 
 	/** @return true if InvalidateProxyOnChange is enabled (default) */
 	bool GetInvalidateProxyOnChangeEnabled() const { return bInvalidateProxyOnChange; }
@@ -294,7 +294,7 @@ protected:
 	FDelegateHandle MeshObjectChangedHandle;
 
 	/** Called whenever internal MeshObject is modified, fires OnMeshChanged and OnMeshVerticesChanged above */
-	void OnMeshObjectChanged(UDynamicMesh* ChangedMeshObject, FDynamicMeshChangeInfo ChangeInfo);
+	GEOMETRYFRAMEWORK_API void OnMeshObjectChanged(UDynamicMesh* ChangedMeshObject, FDynamicMeshChangeInfo ChangeInfo);
 
 
 
@@ -308,24 +308,24 @@ protected:
 	//
 public:
 	/** Set an active triangle color function if one exists, and update the mesh */
-	virtual void SetTriangleColorFunction(TUniqueFunction<FColor(const FDynamicMesh3*, int)> TriangleColorFuncIn,
+	GEOMETRYFRAMEWORK_API virtual void SetTriangleColorFunction(TUniqueFunction<FColor(const FDynamicMesh3*, int)> TriangleColorFuncIn,
 										  EDynamicMeshComponentRenderUpdateMode UpdateMode = EDynamicMeshComponentRenderUpdateMode::FastUpdate);
 
 	/** Clear an active triangle color function if one exists, and update the mesh */
-	virtual void ClearTriangleColorFunction(EDynamicMeshComponentRenderUpdateMode UpdateMode = EDynamicMeshComponentRenderUpdateMode::FastUpdate);
+	GEOMETRYFRAMEWORK_API virtual void ClearTriangleColorFunction(EDynamicMeshComponentRenderUpdateMode UpdateMode = EDynamicMeshComponentRenderUpdateMode::FastUpdate);
 
 	/** @return true if a triangle color function is configured */
-	virtual bool HasTriangleColorFunction();
+	GEOMETRYFRAMEWORK_API virtual bool HasTriangleColorFunction();
 
 protected:
 	/** If this function is set, we will use these colors instead of vertex colors */
 	TUniqueFunction<FColor(const FDynamicMesh3*, int)> TriangleColorFunc = nullptr;
 
 	/** This function is passed via lambda to the RenderProxy to be able to access TriangleColorFunc */
-	FColor GetTriangleColor(const FDynamicMesh3* Mesh, int TriangleID);
+	GEOMETRYFRAMEWORK_API FColor GetTriangleColor(const FDynamicMesh3* Mesh, int TriangleID);
 
 	/** This function is passed via lambda to the RenderProxy when BaseDynamicMeshComponent::ColorMode == Polygroups */
-	FColor GetGroupColor(const FDynamicMesh3* Mesh, int TriangleID) const;
+	GEOMETRYFRAMEWORK_API FColor GetGroupColor(const FDynamicMesh3* Mesh, int TriangleID) const;
 
 
 	//===============================================================================================================
@@ -336,21 +336,21 @@ protected:
 	// minimal vertex buffer updates necessary in the RenderProxy
 public:
 	/** Set an active VertexColor Remapping function if one exists, and update the mesh */
-	virtual void SetVertexColorRemappingFunction(TUniqueFunction<void(FVector4f&)> ColorMapFuncIn,
+	GEOMETRYFRAMEWORK_API virtual void SetVertexColorRemappingFunction(TUniqueFunction<void(FVector4f&)> ColorMapFuncIn,
 		EDynamicMeshComponentRenderUpdateMode UpdateMode = EDynamicMeshComponentRenderUpdateMode::FastUpdate);
 
 	/** Clear an active VertexColor Remapping function if one exists, and update the mesh */
-	virtual void ClearVertexColorRemappingFunction(EDynamicMeshComponentRenderUpdateMode UpdateMode = EDynamicMeshComponentRenderUpdateMode::FastUpdate);
+	GEOMETRYFRAMEWORK_API virtual void ClearVertexColorRemappingFunction(EDynamicMeshComponentRenderUpdateMode UpdateMode = EDynamicMeshComponentRenderUpdateMode::FastUpdate);
 
 	/** @return true if a VertexColor Remapping function is configured */
-	virtual bool HasVertexColorRemappingFunction();
+	GEOMETRYFRAMEWORK_API virtual bool HasVertexColorRemappingFunction();
 
 protected:
 	/** If this function is set, DynamicMesh Attribute Color Overlay colors will be passed through this function before sending to render buffers */
 	TUniqueFunction<void(FVector4f&)> VertexColorMappingFunc = nullptr;
 
 	/** This function is passed via lambda to the RenderProxy to be able to access VertexColorMappingFunc */
-	void RemapVertexColor(FVector4f& VertexColorInOut);
+	GEOMETRYFRAMEWORK_API void RemapVertexColor(FVector4f& VertexColorInOut);
 
 
 
@@ -367,12 +367,12 @@ public:
 	 * function into a second index buffer. These triangles will be drawn with the Secondary render material
 	 * that is set in the BaseDynamicMeshComponent. Calling this function invalidates the SceneProxy.
 	 */
-	virtual void EnableSecondaryTriangleBuffers(TUniqueFunction<bool(const FDynamicMesh3*, int32)> SecondaryTriFilterFunc);
+	GEOMETRYFRAMEWORK_API virtual void EnableSecondaryTriangleBuffers(TUniqueFunction<bool(const FDynamicMesh3*, int32)> SecondaryTriFilterFunc);
 
 	/**
 	 * Disable secondary triangle buffers. This invalidates the SceneProxy.
 	 */
-	virtual void DisableSecondaryTriangleBuffers();
+	GEOMETRYFRAMEWORK_API virtual void DisableSecondaryTriangleBuffers();
 
 protected:
 	TUniqueFunction<bool(const FDynamicMesh3*, int32)> SecondaryTriFilterFunc = nullptr;
@@ -392,7 +392,7 @@ public:
 	 * Configure a decomposition of the mesh, which will result in separate render buffers for each 
 	 * decomposition triangle group. Invalidates existing SceneProxy.
 	 */
-	virtual void SetExternalDecomposition(TUniquePtr<FMeshRenderDecomposition> Decomposition);
+	GEOMETRYFRAMEWORK_API virtual void SetExternalDecomposition(TUniquePtr<FMeshRenderDecomposition> Decomposition);
 
 protected:
 	TUniquePtr<FMeshRenderDecomposition> Decomposition;
@@ -410,17 +410,17 @@ public:
 	/**
 	 * Add a render mesh processor, to be called before the mesh is sent for rendering.
 	 */
-	virtual void SetRenderMeshPostProcessor(TUniquePtr<IRenderMeshPostProcessor> Processor);
+	GEOMETRYFRAMEWORK_API virtual void SetRenderMeshPostProcessor(TUniquePtr<IRenderMeshPostProcessor> Processor);
 
 	/**
 	 * The SceneProxy should call these functions to get the post-processed RenderMesh. (See IRenderMeshPostProcessor.)
 	 */
-	virtual FDynamicMesh3* GetRenderMesh();
+	GEOMETRYFRAMEWORK_API virtual FDynamicMesh3* GetRenderMesh();
 
 	/**
 	 * The SceneProxy should call these functions to get the post-processed RenderMesh. (See IRenderMeshPostProcessor.)
 	 */
-	virtual const FDynamicMesh3* GetRenderMesh() const;
+	GEOMETRYFRAMEWORK_API virtual const FDynamicMesh3* GetRenderMesh() const;
 
 protected:
 	TUniquePtr<IRenderMeshPostProcessor> RenderMeshPostProcessor;
@@ -455,7 +455,7 @@ public:
 	 * different Material IDs on the mesh MaterialID attribute
 	 */
 	UFUNCTION(BlueprintCallable, Category = "Dynamic Mesh Component")
-	void ConfigureMaterialSet(const TArray<UMaterialInterface*>& NewMaterialSet);
+	GEOMETRYFRAMEWORK_API void ConfigureMaterialSet(const TArray<UMaterialInterface*>& NewMaterialSet);
 
 	/**
 	 * Compute the maximum MaterialID on the DynamicMesh, and ensure that Material Slots match.
@@ -465,7 +465,7 @@ public:
 	 * @return true if at the end of this function, Material Slot Count == Max MaterialID
 	 */
 	UFUNCTION(BlueprintCallable, Category = "Dynamic Mesh Component")
-	bool ValidateMaterialSlots(bool bCreateIfMissing = true, bool bDeleteExtraSlots = true);
+	GEOMETRYFRAMEWORK_API bool ValidateMaterialSlots(bool bCreateIfMissing = true, bool bDeleteExtraSlots = true);
 
 
 	//===============================================================================================================
@@ -480,16 +480,16 @@ public:
 	//
 public:
 	UFUNCTION(BlueprintCallable, Category = "Dynamic Mesh Component")
-	void SetTangentsType(EDynamicMeshComponentTangentsMode NewTangentsType);
+	GEOMETRYFRAMEWORK_API void SetTangentsType(EDynamicMeshComponentTangentsMode NewTangentsType);
 
 	UFUNCTION(BlueprintCallable, Category = "Dynamic Mesh Component")
 	EDynamicMeshComponentTangentsMode GetTangentsType() const { return TangentsType; }
 
 	/** This function marks the auto tangents as dirty, they will be recomputed before they are used again */
-	virtual void InvalidateAutoCalculatedTangents();
+	GEOMETRYFRAMEWORK_API virtual void InvalidateAutoCalculatedTangents();
 
 	/** @return AutoCalculated Tangent Set, which may require that they be recomputed, or nullptr if not enabled/available */
-	const UE::Geometry::FMeshTangentsf* GetAutoCalculatedTangents();
+	GEOMETRYFRAMEWORK_API const UE::Geometry::FMeshTangentsf* GetAutoCalculatedTangents();
 
 protected:
 	/** How should Tangents be calculated/handled */
@@ -502,7 +502,7 @@ protected:
 	/** Set of per-triangle-vertex tangents computed for the current mesh. Only valid if bAutoCalculatedTangentsValid == true */
 	UE::Geometry::FMeshTangentsf AutoCalculatedTangents;
 
-	void UpdateAutoCalculatedTangents();
+	GEOMETRYFRAMEWORK_API void UpdateAutoCalculatedTangents();
 
 
 
@@ -515,7 +515,7 @@ public:
 	 * calls SetComplexAsSimpleCollisionEnabled(true, true)
 	 */
 	UFUNCTION(BlueprintCallable, Category = "Dynamic Mesh Component")
-	void EnableComplexAsSimpleCollision();
+	GEOMETRYFRAMEWORK_API void EnableComplexAsSimpleCollision();
 
 	/**
 	 * If bEnabled=true, sets bEnableComplexCollision=true and CollisionType=CTF_UseComplexAsSimple
@@ -523,31 +523,31 @@ public:
 	 * @param bImmediateUpdate if true, UpdateCollision(true) is called
 	 */
 	UFUNCTION(BlueprintCallable, Category = "Dynamic Mesh Component")
-	void SetComplexAsSimpleCollisionEnabled(bool bEnabled, bool bImmediateUpdate = true);
+	GEOMETRYFRAMEWORK_API void SetComplexAsSimpleCollisionEnabled(bool bEnabled, bool bImmediateUpdate = true);
 
 	/**
 	 * Set value of bDeferCollisionUpdates, when enabled, collision is not automatically recomputed each time the mesh changes.
 	 * @param bImmediateUpdate if true, UpdateCollision(true) is called if bEnabled=false, ie to force a collision rebuild
 	 */
 	UFUNCTION(BlueprintCallable, Category = "Dynamic Mesh Component")
-	void SetDeferredCollisionUpdatesEnabled(bool bEnabled, bool bImmediateUpdate = true);
+	GEOMETRYFRAMEWORK_API void SetDeferredCollisionUpdatesEnabled(bool bEnabled, bool bImmediateUpdate = true);
 
 
-	virtual bool GetPhysicsTriMeshData(struct FTriMeshCollisionData* CollisionData, bool InUseAllTriData) override;
-	virtual bool ContainsPhysicsTriMeshData(bool InUseAllTriData) const override;
-	virtual bool WantsNegXTriMesh() override;
+	GEOMETRYFRAMEWORK_API virtual bool GetPhysicsTriMeshData(struct FTriMeshCollisionData* CollisionData, bool InUseAllTriData) override;
+	GEOMETRYFRAMEWORK_API virtual bool ContainsPhysicsTriMeshData(bool InUseAllTriData) const override;
+	GEOMETRYFRAMEWORK_API virtual bool WantsNegXTriMesh() override;
 
 	/** @return current BodySetup for this Component, or nullptr if it does not exist */
 	virtual const UBodySetup* GetBodySetup() const { return MeshBodySetup; }
 	/** @return BodySetup for this Component. A new BodySetup will be created if one does not exist. */
-	virtual UBodySetup* GetBodySetup() override;
+	GEOMETRYFRAMEWORK_API virtual UBodySetup* GetBodySetup() override;
 
 	/**
 	 * Force an update of the Collision/Physics data for this Component.
 	 * @param bOnlyIfPending only update if a collision update is pending, ie the underlying DynamicMesh changed and bDeferCollisionUpdates is enabled
 	 */
 	UFUNCTION(BlueprintCallable, Category = "Dynamic Mesh Component")
-	virtual void UpdateCollision(bool bOnlyIfPending = true);
+	GEOMETRYFRAMEWORK_API virtual void UpdateCollision(bool bOnlyIfPending = true);
 
 	/** Type of Collision Geometry to use for this Mesh */
 	UPROPERTY(EditAnywhere, BlueprintReadWrite, Category = "Dynamic Mesh Component|Collision")
@@ -559,7 +559,7 @@ public:
 	 * @param AggGeom				New simple collision shapes to be used
 	 * @param bUpdateCollision		Whether to automatically call UpdateCollision() -- if false, manually call it to register the change with the physics system
 	 */
-	virtual void SetSimpleCollisionShapes(const struct FKAggregateGeom& AggGeom, bool bUpdateCollision);
+	GEOMETRYFRAMEWORK_API virtual void SetSimpleCollisionShapes(const struct FKAggregateGeom& AggGeom, bool bUpdateCollision);
 
 	const FKAggregateGeom& GetSimpleCollisionShapes() const
 	{
@@ -570,7 +570,7 @@ public:
 	 * Clear the simple collision shapes associated with this mesh component
 	 * @param bUpdateCollision		Whether to automatically call UpdateCollision() -- if false, manually call it to register the change with the physics system
 	 */
-	virtual void ClearSimpleCollisionShapes(bool bUpdateCollision);
+	GEOMETRYFRAMEWORK_API virtual void ClearSimpleCollisionShapes(bool bUpdateCollision);
 
 	/**
 	 *	Controls whether the physics cooking should be done off the game thread.
@@ -596,14 +596,14 @@ public:
 	 * want to defer collision updates during the user interaction, but not change the serialized property. 
 	 * In this case a non-serialized version of the flag can be *temporarily* set via this function in C++. 
 	 */
-	void SetTransientDeferCollisionUpdates(bool bEnabled);
+	GEOMETRYFRAMEWORK_API void SetTransientDeferCollisionUpdates(bool bEnabled);
 
 protected:
 	UPROPERTY(Instanced)
 	TObjectPtr<UBodySetup> MeshBodySetup;
 
-	virtual void InvalidatePhysicsData();
-	virtual void RebuildPhysicsData();
+	GEOMETRYFRAMEWORK_API virtual void InvalidatePhysicsData();
+	GEOMETRYFRAMEWORK_API virtual void RebuildPhysicsData();
 
 	bool bTransientDeferCollisionUpdates = false;
 	bool bCollisionUpdatePending = false;
@@ -618,7 +618,7 @@ protected:
 	UE::Geometry::FAxisAlignedBox3d LocalBounds;
 
 	/** Recompute LocalBounds from the current Mesh */
-	void UpdateLocalBounds();
+	GEOMETRYFRAMEWORK_API void UpdateLocalBounds();
 
 	//
 	// Internals for managing collision representation and setup
@@ -633,18 +633,18 @@ protected:
 	TArray<TObjectPtr<UBodySetup>> AsyncBodySetupQueue;
 
 	/** Once async physics cook is done, create needed state */
-	virtual void FinishPhysicsAsyncCook(bool bSuccess, UBodySetup* FinishedBodySetup);
+	GEOMETRYFRAMEWORK_API virtual void FinishPhysicsAsyncCook(bool bSuccess, UBodySetup* FinishedBodySetup);
 
-	virtual UBodySetup* CreateBodySetupHelper();
+	GEOMETRYFRAMEWORK_API virtual UBodySetup* CreateBodySetupHelper();
 	
 	/** @return Set new BodySetup for this Component. */
-	virtual void SetBodySetup(UBodySetup* NewSetup);
+	GEOMETRYFRAMEWORK_API virtual void SetBodySetup(UBodySetup* NewSetup);
 
 	/**
 	 * This is called to tell our RenderProxy about modifications to the material set.
 	 * We need to pass this on for things like material validation in the Editor.
 	 */
-	virtual void NotifyMaterialSetUpdated();
+	GEOMETRYFRAMEWORK_API virtual void NotifyMaterialSetUpdated();
 
 	/**
 	 * If the render proxy is invalidated (eg by MarkRenderStateDirty()), it will be destroyed at the end of
@@ -668,32 +668,32 @@ protected:
 	/** 
 	 * @return current render proxy, if valid, otherwise nullptr 
 	 */
-	FDynamicMeshSceneProxy* GetCurrentSceneProxy();
+	GEOMETRYFRAMEWORK_API FDynamicMeshSceneProxy* GetCurrentSceneProxy();
 
 
 	/**
 	 * Fully invalidate all rendering data for this Component. Current Proxy will be discarded, Bounds and possibly Tangents recomputed, etc
 	 */
-	void ResetProxy();
+	GEOMETRYFRAMEWORK_API void ResetProxy();
 
 	//~ Begin UPrimitiveComponent Interface.
-	virtual FPrimitiveSceneProxy* CreateSceneProxy() override;
+	GEOMETRYFRAMEWORK_API virtual FPrimitiveSceneProxy* CreateSceneProxy() override;
 
 	//~ USceneComponent Interface.
-	virtual FBoxSphereBounds CalcBounds(const FTransform& LocalToWorld) const override;
-	virtual void OnChildAttached(USceneComponent* ChildComponent) override;
-	virtual void OnChildDetached(USceneComponent* ChildComponent) override;
+	GEOMETRYFRAMEWORK_API virtual FBoxSphereBounds CalcBounds(const FTransform& LocalToWorld) const override;
+	GEOMETRYFRAMEWORK_API virtual void OnChildAttached(USceneComponent* ChildComponent) override;
+	GEOMETRYFRAMEWORK_API virtual void OnChildDetached(USceneComponent* ChildComponent) override;
 
 	//~ UObject Interface.
-	virtual void PostLoad() override;
-	virtual void BeginDestroy() override;
+	GEOMETRYFRAMEWORK_API virtual void PostLoad() override;
+	GEOMETRYFRAMEWORK_API virtual void BeginDestroy() override;
 #if WITH_EDITOR
-	void PostEditChangeProperty(FPropertyChangedEvent& PropertyChangedEvent) override;
+	GEOMETRYFRAMEWORK_API void PostEditChangeProperty(FPropertyChangedEvent& PropertyChangedEvent) override;
 #endif
 
 public:
 	/** Set whether or not to validate mesh batch materials against the component materials. */
-	void SetSceneProxyVerifyUsedMaterials(bool bState);
+	GEOMETRYFRAMEWORK_API void SetSceneProxyVerifyUsedMaterials(bool bState);
 
 
 };
diff --git a/Engine/Source/Runtime/GeometryFramework/Public/Components/MeshRenderDecomposition.h b/Engine/Source/Runtime/GeometryFramework/Public/Components/MeshRenderDecomposition.h
--- a/Engine/Source/Runtime/GeometryFramework/Public/Components/MeshRenderDecomposition.h
+++ b/Engine/Source/Runtime/GeometryFramework/Public/Components/MeshRenderDecomposition.h
@@ -16,7 +16,7 @@ struct FComponentMaterialSet;
  * FMeshRenderDecomposition represents a decomposition of a mesh into "chunks" of triangles, with associated materials.
  * This is passed to the rendering components to split a mesh into multiple RenderBuffers, for more efficient updating.
  */
-class GEOMETRYFRAMEWORK_API FMeshRenderDecomposition
+class FMeshRenderDecomposition
 {
 public:
 
@@ -84,17 +84,17 @@ public:
 	/**
 	 * Construct mappings between mesh and groups (eg TriangleToGroupMap)
 	 */
-	void BuildAssociations(const FDynamicMesh3* Mesh);
+	GEOMETRYFRAMEWORK_API void BuildAssociations(const FDynamicMesh3* Mesh);
 
 
 	/**
 	 * Build decomposition with one group for each MaterialID of mesh
 	 */
-	static void BuildMaterialDecomposition(const FDynamicMesh3* Mesh, const FComponentMaterialSet* MaterialSet, FMeshRenderDecomposition& Decomp);
+	static GEOMETRYFRAMEWORK_API void BuildMaterialDecomposition(const FDynamicMesh3* Mesh, const FComponentMaterialSet* MaterialSet, FMeshRenderDecomposition& Decomp);
 
 	/**
 	 * Build per-material decomposition, and then split each of those into chunks of at most MaxChunkSize
 	 * (actual chunk sizes will be highly variable and some may be very small...)
 	 */
-	static void BuildChunkedDecomposition(const FDynamicMesh3* Mesh, const FComponentMaterialSet* MaterialSet, FMeshRenderDecomposition& Decomp, int32 MaxChunkSize = 1 << 14 /* 16k */ );
+	static GEOMETRYFRAMEWORK_API void BuildChunkedDecomposition(const FDynamicMesh3* Mesh, const FComponentMaterialSet* MaterialSet, FMeshRenderDecomposition& Decomp, int32 MaxChunkSize = 1 << 14 /* 16k */ );
 };
diff --git a/Engine/Source/Runtime/GeometryFramework/Public/DynamicMeshActor.h b/Engine/Source/Runtime/GeometryFramework/Public/DynamicMeshActor.h
--- a/Engine/Source/Runtime/GeometryFramework/Public/DynamicMeshActor.h
+++ b/Engine/Source/Runtime/GeometryFramework/Public/DynamicMeshActor.h
@@ -12,8 +12,8 @@
 /**
  * ADynamicMeshActor is an Actor that has a USimpleDynamicMeshComponent as it's RootObject.
  */
-UCLASS(ConversionRoot, ComponentWrapperClass, ClassGroup=DynamicMesh, meta = (ChildCanTick))
-class GEOMETRYFRAMEWORK_API ADynamicMeshActor : public AActor
+UCLASS(ConversionRoot, ComponentWrapperClass, ClassGroup=DynamicMesh, meta = (ChildCanTick), MinimalAPI)
+class ADynamicMeshActor : public AActor
 {
 	GENERATED_UCLASS_BODY()
 
@@ -46,23 +46,23 @@ protected:
 public:
 	/** Access the compute mesh pool */
 	UFUNCTION(BlueprintCallable, Category = DynamicMeshActor)
-	UDynamicMeshPool* GetComputeMeshPool();
+	GEOMETRYFRAMEWORK_API UDynamicMeshPool* GetComputeMeshPool();
 
 	/** Request a compute mesh from the Pool, which will return a previously-allocated mesh or add and return a new one. If the Pool is disabled, a new UDynamicMesh will be allocated and returned. */
 	UFUNCTION(BlueprintCallable, Category = DynamicMeshActor)
-	UDynamicMesh* AllocateComputeMesh();
+	GEOMETRYFRAMEWORK_API UDynamicMesh* AllocateComputeMesh();
 
 	/** Release a compute mesh back to the Pool */
 	UFUNCTION(BlueprintCallable, Category = DynamicMeshActor)
-	bool ReleaseComputeMesh(UDynamicMesh* Mesh);
+	GEOMETRYFRAMEWORK_API bool ReleaseComputeMesh(UDynamicMesh* Mesh);
 
 	/** Release all compute meshes that the Pool has allocated */
 	UFUNCTION(BlueprintCallable, Category = DynamicMeshActor)
-	void ReleaseAllComputeMeshes();
+	GEOMETRYFRAMEWORK_API void ReleaseAllComputeMeshes();
 
 	/** Release all compute meshes that the Pool has allocated, and then release them from the Pool, so that they will be garbage-collected */
 	UFUNCTION(BlueprintCallable, Category = DynamicMeshActor)
-	void FreeAllComputeMeshes();
+	GEOMETRYFRAMEWORK_API void FreeAllComputeMeshes();
 };
 
 
diff --git a/Engine/Source/Runtime/HeadMountedDisplay/Public/HeadMountedDisplayTypes.h b/Engine/Source/Runtime/HeadMountedDisplay/Public/HeadMountedDisplayTypes.h
--- a/Engine/Source/Runtime/HeadMountedDisplay/Public/HeadMountedDisplayTypes.h
+++ b/Engine/Source/Runtime/HeadMountedDisplay/Public/HeadMountedDisplayTypes.h
@@ -30,7 +30,7 @@ class UObject;
 struct FFilterVertex;
 struct FFrame;
 
-class HEADMOUNTEDDISPLAY_API FHMDViewMesh
+class FHMDViewMesh
 {
 public:
 
@@ -40,15 +40,15 @@ public:
 		MT_VisibleArea
 	};
 
-	FHMDViewMesh();
-	~FHMDViewMesh();
+	HEADMOUNTEDDISPLAY_API FHMDViewMesh();
+	HEADMOUNTEDDISPLAY_API ~FHMDViewMesh();
 
 	bool IsValid() const
 	{
 		return NumTriangles > 0;
 	}
 
-	void BuildMesh(const FVector2D Positions[], uint32 VertexCount, EHMDMeshType MeshType);
+	HEADMOUNTEDDISPLAY_API void BuildMesh(const FVector2D Positions[], uint32 VertexCount, EHMDMeshType MeshType);
 
 	FBufferRHIRef VertexBufferRHI;
 	FBufferRHIRef IndexBufferRHI;
@@ -287,8 +287,8 @@ enum class EHandKeypoint : uint8
 
 const int32 EHandKeypointCount = static_cast<int32>(EHandKeypoint::LittleTip) + 1;
 
-UCLASS()
-class HEADMOUNTEDDISPLAY_API UHandKeypointConversion : public UBlueprintFunctionLibrary
+UCLASS(MinimalAPI)
+class UHandKeypointConversion : public UBlueprintFunctionLibrary
 {
 	GENERATED_BODY()
 
@@ -311,7 +311,7 @@ enum class EXRVisualType : uint8
 };
 
 USTRUCT(BlueprintType)
-struct HEADMOUNTEDDISPLAY_API FXRHMDData
+struct FXRHMDData
 {
 	GENERATED_USTRUCT_BODY();
 
@@ -332,7 +332,7 @@ struct HEADMOUNTEDDISPLAY_API FXRHMDData
 };
 
 USTRUCT(BlueprintType)
-struct HEADMOUNTEDDISPLAY_API FXRMotionControllerData
+struct FXRMotionControllerData
 {
 	GENERATED_USTRUCT_BODY();
 
diff --git a/Engine/Source/Runtime/HeadMountedDisplay/Public/IHandTracker.h b/Engine/Source/Runtime/HeadMountedDisplay/Public/IHandTracker.h
--- a/Engine/Source/Runtime/HeadMountedDisplay/Public/IHandTracker.h
+++ b/Engine/Source/Runtime/HeadMountedDisplay/Public/IHandTracker.h
@@ -16,7 +16,7 @@ enum class EHandKeypoint : uint8;
 /**
  */
 
-class HEADMOUNTEDDISPLAY_API IHandTracker : public IModularFeature
+class IHandTracker : public IModularFeature
 {
 public:
 	static FName GetModularFeatureName()
diff --git a/Engine/Source/Runtime/HeadMountedDisplay/Public/IHeadMountedDisplay.h b/Engine/Source/Runtime/HeadMountedDisplay/Public/IHeadMountedDisplay.h
--- a/Engine/Source/Runtime/HeadMountedDisplay/Public/IHeadMountedDisplay.h
+++ b/Engine/Source/Runtime/HeadMountedDisplay/Public/IHeadMountedDisplay.h
@@ -34,11 +34,11 @@ struct FHeadMountedDisplayPassContext
  * HMD device interface
  */
 
-class HEADMOUNTEDDISPLAY_API IHeadMountedDisplay : public IModuleInterface
+class IHeadMountedDisplay : public IModuleInterface
 {
 
 public:
-	IHeadMountedDisplay();
+	HEADMOUNTEDDISPLAY_API IHeadMountedDisplay();
 
 	/**
 	 * Returns true if HMD is currently connected.  It may or may not be in use.
@@ -255,12 +255,12 @@ public:
 	 * the regular desktop apps. In this case, FCoreDelegates::ApplicationWillEnterBackgroundDelegate and FCoreDelegates::ApplicationHasEnteredForegroundDelegate
 	 * reflect the state of VR focus (either the app should be rendered in HMD or not).
 	 */
-	virtual bool DoesAppUseVRFocus() const;
+	HEADMOUNTEDDISPLAY_API virtual bool DoesAppUseVRFocus() const;
 
 	/**
 	 * Returns true, if the app has VR focus, meaning if it is rendered in the HMD.
 	 */
-	virtual bool DoesAppHaveVRFocus() const;
+	HEADMOUNTEDDISPLAY_API virtual bool DoesAppHaveVRFocus() const;
 
 	/**
 	 * If true, scene rendering should be skipped.
diff --git a/Engine/Source/Runtime/HeadMountedDisplay/Public/IHeadMountedDisplayVulkanExtensions.h b/Engine/Source/Runtime/HeadMountedDisplay/Public/IHeadMountedDisplayVulkanExtensions.h
--- a/Engine/Source/Runtime/HeadMountedDisplay/Public/IHeadMountedDisplayVulkanExtensions.h
+++ b/Engine/Source/Runtime/HeadMountedDisplay/Public/IHeadMountedDisplayVulkanExtensions.h
@@ -1,4 +1,4 @@
-﻿// Copyright Epic Games, Inc. All Rights Reserved.
+// Copyright Epic Games, Inc. All Rights Reserved.
 
 #pragma once
 
@@ -9,7 +9,7 @@
 *
 * Knowledge of these extensions is required during Vulkan RHI initialization.
 */
-class HEADMOUNTEDDISPLAY_API IHeadMountedDisplayVulkanExtensions
+class IHeadMountedDisplayVulkanExtensions
 {
 public:
 	virtual bool GetVulkanInstanceExtensionsRequired(TArray<const ANSICHAR*>& Out) = 0;
@@ -19,5 +19,5 @@ public:
 	* If true and VK_PRESENT_MODE_MAILBOX_KHR is not available, VK_PRESENT_MODE_IMMEDIATE_KHR will be preferred over VK_PRESENT_MODE_FIFO_KHR
 	* when creating the swap chain for the spectator window.
 	*/
-	virtual bool ShouldDisableVulkanVSync() const;
+	HEADMOUNTEDDISPLAY_API virtual bool ShouldDisableVulkanVSync() const;
 };
diff --git a/Engine/Source/Runtime/HeadMountedDisplay/Public/IIdentifiableXRDevice.h b/Engine/Source/Runtime/HeadMountedDisplay/Public/IIdentifiableXRDevice.h
--- a/Engine/Source/Runtime/HeadMountedDisplay/Public/IIdentifiableXRDevice.h
+++ b/Engine/Source/Runtime/HeadMountedDisplay/Public/IIdentifiableXRDevice.h
@@ -10,7 +10,7 @@
 
 #include "IIdentifiableXRDevice.generated.h"
 
-class HEADMOUNTEDDISPLAY_API IXRSystemIdentifier
+class IXRSystemIdentifier
 {
 public:
 	/**
@@ -33,7 +33,7 @@ public:
  * interfaces which both reference the same devices, then this base class gives 
  * you a way to communicate between the two.
  */
-class HEADMOUNTEDDISPLAY_API IIdentifiableXRDevice : public IXRSystemIdentifier
+class IIdentifiableXRDevice : public IXRSystemIdentifier
 {
 public:
 	/**
@@ -54,17 +54,17 @@ public:
 };
 
 USTRUCT(BlueprintType)
-struct HEADMOUNTEDDISPLAY_API FXRDeviceId
+struct FXRDeviceId
 {
 	GENERATED_USTRUCT_BODY()
 
 public:
-	FXRDeviceId();
-	FXRDeviceId(IIdentifiableXRDevice* DeviceId);
-	FXRDeviceId(IXRSystemIdentifier* OwningSystem, const int32 DeviceId);
+	HEADMOUNTEDDISPLAY_API FXRDeviceId();
+	HEADMOUNTEDDISPLAY_API FXRDeviceId(IIdentifiableXRDevice* DeviceId);
+	HEADMOUNTEDDISPLAY_API FXRDeviceId(IXRSystemIdentifier* OwningSystem, const int32 DeviceId);
 
-	bool IsOwnedBy(IXRSystemIdentifier* XRSystem) const;
-	void Clear();
+	HEADMOUNTEDDISPLAY_API bool IsOwnedBy(IXRSystemIdentifier* XRSystem) const;
+	HEADMOUNTEDDISPLAY_API void Clear();
 	
 	bool IsSet() const 
 	{ 
@@ -73,8 +73,8 @@ public:
 		return !SystemName.IsNone(); 
 	}
 
-	bool operator==(const FXRDeviceId& Rhs) const;
-	bool operator==(const IIdentifiableXRDevice* Rhs) const;
+	HEADMOUNTEDDISPLAY_API bool operator==(const FXRDeviceId& Rhs) const;
+	HEADMOUNTEDDISPLAY_API bool operator==(const IIdentifiableXRDevice* Rhs) const;
 
 public:
 	UPROPERTY(BlueprintReadOnly, Category="XRDevice")
diff --git a/Engine/Source/Runtime/HeadMountedDisplay/Public/IMotionController.h b/Engine/Source/Runtime/HeadMountedDisplay/Public/IMotionController.h
--- a/Engine/Source/Runtime/HeadMountedDisplay/Public/IMotionController.h
+++ b/Engine/Source/Runtime/HeadMountedDisplay/Public/IMotionController.h
@@ -44,7 +44,7 @@ struct FMotionControllerSource
  *			to be both piggy-backed off HMD devices which support them, as well as standing alone.
  */
 
-class HEADMOUNTEDDISPLAY_API IMotionController : public IModularFeature
+class IMotionController : public IModularFeature
 {
 public:
 	virtual ~IMotionController() {}
@@ -136,9 +136,9 @@ public:
 	virtual bool SetPlayerMappableInputConfig(TObjectPtr<class UPlayerMappableInputConfig> InputConfig = nullptr) { return true; };
 
 	// explicit source names
-	static FName LeftHandSourceId;
-	static FName RightHandSourceId;
-	static FName HMDSourceId;
+	static HEADMOUNTEDDISPLAY_API FName LeftHandSourceId;
+	static HEADMOUNTEDDISPLAY_API FName RightHandSourceId;
+	static HEADMOUNTEDDISPLAY_API FName HMDSourceId;
 
-	static bool GetHandEnumForSourceName(const FName Source, EControllerHand& OutHand);
+	static HEADMOUNTEDDISPLAY_API bool GetHandEnumForSourceName(const FName Source, EControllerHand& OutHand);
 };
diff --git a/Engine/Source/Runtime/HeadMountedDisplay/Public/IMotionTrackingSystemManagement.h b/Engine/Source/Runtime/HeadMountedDisplay/Public/IMotionTrackingSystemManagement.h
--- a/Engine/Source/Runtime/HeadMountedDisplay/Public/IMotionTrackingSystemManagement.h
+++ b/Engine/Source/Runtime/HeadMountedDisplay/Public/IMotionTrackingSystemManagement.h
@@ -14,7 +14,7 @@ enum class EControllerHand : uint8;
  * This exposes management options for a motion tracking system.  Some platforms will not implement it.  
  */
 
-class HEADMOUNTEDDISPLAY_API IMotionTrackingSystemManagement : public IModularFeature
+class IMotionTrackingSystemManagement : public IModularFeature
 {
 public:
 	static FName GetModularFeatureName()
diff --git a/Engine/Source/Runtime/HeadMountedDisplay/Public/ISpectatorScreenController.h b/Engine/Source/Runtime/HeadMountedDisplay/Public/ISpectatorScreenController.h
--- a/Engine/Source/Runtime/HeadMountedDisplay/Public/ISpectatorScreenController.h
+++ b/Engine/Source/Runtime/HeadMountedDisplay/Public/ISpectatorScreenController.h
@@ -16,7 +16,7 @@ struct FSpectatorScreenModeTexturePlusEyeLayout;
  * This is the interface to control the spectator screen, on platforms that support it.
  */
 
-class HEADMOUNTEDDISPLAY_API ISpectatorScreenController : public IModularFeature
+class ISpectatorScreenController : public IModularFeature
 {
 public:
 	virtual ~ISpectatorScreenController() {}
diff --git a/Engine/Source/Runtime/HeadMountedDisplay/Public/IXRCamera.h b/Engine/Source/Runtime/HeadMountedDisplay/Public/IXRCamera.h
--- a/Engine/Source/Runtime/HeadMountedDisplay/Public/IXRCamera.h
+++ b/Engine/Source/Runtime/HeadMountedDisplay/Public/IXRCamera.h
@@ -18,7 +18,7 @@ class USceneComponent;
  *
  * To reduce duplicated code, implementations should use the FXRCameraBase instead of implementing this interface directly.
  */
-class HEADMOUNTEDDISPLAY_API  IXRCamera : public IIdentifiableXRDevice
+class  IXRCamera : public IIdentifiableXRDevice
 {
 public:
 	IXRCamera() {}
diff --git a/Engine/Source/Runtime/HeadMountedDisplay/Public/IXRInput.h b/Engine/Source/Runtime/HeadMountedDisplay/Public/IXRInput.h
--- a/Engine/Source/Runtime/HeadMountedDisplay/Public/IXRInput.h
+++ b/Engine/Source/Runtime/HeadMountedDisplay/Public/IXRInput.h
@@ -10,7 +10,7 @@
 /** 
  * Optional interface returned from IXRTrackingSystem if the plugin requires being able to grab touch or keyboard input events.
  */
-class HEADMOUNTEDDISPLAY_API IXRInput
+class IXRInput
 {
 public:
 	/**
diff --git a/Engine/Source/Runtime/HeadMountedDisplay/Public/IXRLoadingScreen.h b/Engine/Source/Runtime/HeadMountedDisplay/Public/IXRLoadingScreen.h
--- a/Engine/Source/Runtime/HeadMountedDisplay/Public/IXRLoadingScreen.h
+++ b/Engine/Source/Runtime/HeadMountedDisplay/Public/IXRLoadingScreen.h
@@ -13,7 +13,7 @@
  *
  * To reduce duplicated code, implementations should use the FXRLoadingScreenBase of FDefaultXRLoadingScreen instead of implementing this interface directly.
  */
-class HEADMOUNTEDDISPLAY_API  IXRLoadingScreen 
+class  IXRLoadingScreen 
 {
 public:
 
@@ -91,5 +91,5 @@ public:
 	 * Internal utility method for implementing backwards compatibility with IStereoLayers::Show/HideSplashScreen.
 	 * Should be called from implementations overriding IStereoLayer::UpdateSplashScreen()
 	 */
-	static void ShowLoadingScreen_Compat(bool bShow, FTexture2DRHIRef Texture, const FVector& Offset, const FVector2D& Scale);
+	static HEADMOUNTEDDISPLAY_API void ShowLoadingScreen_Compat(bool bShow, FTexture2DRHIRef Texture, const FVector& Offset, const FVector2D& Scale);
 };
diff --git a/Engine/Source/Runtime/HeadMountedDisplay/Public/IXRSystemAssets.h b/Engine/Source/Runtime/HeadMountedDisplay/Public/IXRSystemAssets.h
--- a/Engine/Source/Runtime/HeadMountedDisplay/Public/IXRSystemAssets.h
+++ b/Engine/Source/Runtime/HeadMountedDisplay/Public/IXRSystemAssets.h
@@ -13,7 +13,7 @@ enum class EControllerHand : uint8;
 
 DECLARE_DELEGATE_OneParam(FXRComponentLoadComplete, UPrimitiveComponent*);
 
-class HEADMOUNTEDDISPLAY_API IXRSystemAssets : public IModularFeature, public IXRSystemIdentifier
+class IXRSystemAssets : public IModularFeature, public IXRSystemIdentifier
 {
 public:
 	static FName GetModularFeatureName()
diff --git a/Engine/Source/Runtime/HeadMountedDisplay/Public/LateUpdateManager.h b/Engine/Source/Runtime/HeadMountedDisplay/Public/LateUpdateManager.h
--- a/Engine/Source/Runtime/HeadMountedDisplay/Public/LateUpdateManager.h
+++ b/Engine/Source/Runtime/HeadMountedDisplay/Public/LateUpdateManager.h
@@ -10,27 +10,27 @@ class FPrimitiveSceneInfo;
 /**
 * Utility class for applying an offset to a hierarchy of components in the renderer thread.
 */
-class HEADMOUNTEDDISPLAY_API FLateUpdateManager
+class FLateUpdateManager
 {
 public:
-	FLateUpdateManager();
+	HEADMOUNTEDDISPLAY_API FLateUpdateManager();
 	virtual ~FLateUpdateManager() {}
 
 	/** Setup state for applying the render thread late update */
-	void Setup(const FTransform& ParentToWorld, USceneComponent* Component, bool bSkipLateUpdate);
+	HEADMOUNTEDDISPLAY_API void Setup(const FTransform& ParentToWorld, USceneComponent* Component, bool bSkipLateUpdate);
 
 	/** Returns true if the LateUpdateSetup data is stale. */
 	bool GetSkipLateUpdate_RenderThread() const { return UpdateStates[LateUpdateRenderReadIndex].bSkip; }
 
 	/** Apply the late update delta to the cached components */
-	void Apply_RenderThread(FSceneInterface* Scene, const FTransform& OldRelativeTransform, const FTransform& NewRelativeTransform);
+	HEADMOUNTEDDISPLAY_API void Apply_RenderThread(FSceneInterface* Scene, const FTransform& OldRelativeTransform, const FTransform& NewRelativeTransform);
 
 private:
 
 	/** A utility method that calls CacheSceneInfo on ParentComponent and all of its descendants */
-	void GatherLateUpdatePrimitives(USceneComponent* ParentComponent);
+	HEADMOUNTEDDISPLAY_API void GatherLateUpdatePrimitives(USceneComponent* ParentComponent);
 	/** Generates a LateUpdatePrimitiveInfo for the given component if it has a SceneProxy and appends it to the current LateUpdatePrimitives array */
-	void CacheSceneInfo(USceneComponent* Component);
+	HEADMOUNTEDDISPLAY_API void CacheSceneInfo(USceneComponent* Component);
 
 	struct FLateUpdateState
 	{
diff --git a/Engine/Source/Runtime/HeadMountedDisplay/Public/MotionControllerComponent.h b/Engine/Source/Runtime/HeadMountedDisplay/Public/MotionControllerComponent.h
--- a/Engine/Source/Runtime/HeadMountedDisplay/Public/MotionControllerComponent.h
+++ b/Engine/Source/Runtime/HeadMountedDisplay/Public/MotionControllerComponent.h
@@ -15,8 +15,8 @@ class FSceneView;
 class FSceneViewFamily;
 class UXRDeviceVisualizationComponent;
 
-UCLASS(Blueprintable, meta = (BlueprintSpawnableComponent), ClassGroup = MotionController)
-class HEADMOUNTEDDISPLAY_API UMotionControllerComponent : public UPrimitiveComponent
+UCLASS(Blueprintable, meta = (BlueprintSpawnableComponent), ClassGroup = MotionController, MinimalAPI)
+class UMotionControllerComponent : public UPrimitiveComponent
 {
 	GENERATED_UCLASS_BODY()
 
@@ -46,7 +46,7 @@ class HEADMOUNTEDDISPLAY_API UMotionControllerComponent : public UPrimitiveCompo
 	UPROPERTY(EditAnywhere, BlueprintReadWrite, BlueprintSetter = SetDisplayModelSource, Category = "Visualization", meta = (editcondition = "bDisplayDeviceModel", DeprecatedProperty, DeprecationMessage = "DisplayModelSource is deprecated. Please use the XRDeviceVisualizationComponent for rendering instead."))
 	FName DisplayModelSource;
 
-	static FName CustomModelSourceId;
+	static HEADMOUNTEDDISPLAY_API FName CustomModelSourceId;
 
 	/** A mesh override that'll be displayed attached to this MotionController. */
 	UE_DEPRECATED(5.2, "CustomDisplayMesh is deprecated. Please use the XRDeviceVisualizationComponent for rendering instead.")
@@ -60,15 +60,15 @@ class HEADMOUNTEDDISPLAY_API UMotionControllerComponent : public UPrimitiveCompo
 
 	UE_DEPRECATED(5.2, "SetShowDeviceModel is deprecated. Please use the XRDeviceVisualizationComponent for rendering instead.")
 	UFUNCTION(BlueprintSetter, meta = (DeprecatedFunction, DeprecationMessage = "SetShowDeviceModel is deprecated. Please use the XRDeviceVisualizationComponent for rendering instead."))
-	void SetShowDeviceModel(const bool bShowControllerModel);
+	HEADMOUNTEDDISPLAY_API void SetShowDeviceModel(const bool bShowControllerModel);
 
 	UE_DEPRECATED(5.2, "SetDisplayModelSource is deprecated. Please use the XRDeviceVisualizationComponent for rendering instead.")
 	UFUNCTION(BlueprintSetter, meta = (DeprecatedFunction, DeprecationMessage = "SetDisplayModelSource is deprecated. Please use the XRDeviceVisualizationComponent for rendering instead."))
-	void SetDisplayModelSource(const FName NewDisplayModelSource);
+	HEADMOUNTEDDISPLAY_API void SetDisplayModelSource(const FName NewDisplayModelSource);
 
 	UE_DEPRECATED(5.2, "SetCustomDisplayMesh is deprecated. Please use the XRDeviceVisualizationComponent for rendering instead.")
 	UFUNCTION(BlueprintSetter, meta = (DeprecatedFunction, DeprecationMessage = "SetCustomDisplayMesh is deprecated. Please use the XRDeviceVisualizationComponent for rendering instead."))
-	void SetCustomDisplayMesh(UStaticMesh* NewDisplayMesh);
+	HEADMOUNTEDDISPLAY_API void SetCustomDisplayMesh(UStaticMesh* NewDisplayMesh);
 
 	/** Whether or not this component had a valid tracked device this frame */
 	UFUNCTION(BlueprintPure, Category = "MotionController")
@@ -78,28 +78,28 @@ class HEADMOUNTEDDISPLAY_API UMotionControllerComponent : public UPrimitiveCompo
 	}
 
 	UFUNCTION(BlueprintSetter, meta = (DeprecatedFunction, DeprecationMessage = "Please use the Motion Source property instead of Hand"))
-	void SetTrackingSource(const EControllerHand NewSource);
+	HEADMOUNTEDDISPLAY_API void SetTrackingSource(const EControllerHand NewSource);
 
 	UFUNCTION(BlueprintGetter, meta = (DeprecatedFunction, DeprecationMessage = "Please use the Motion Source property instead of Hand"))
-	EControllerHand GetTrackingSource() const;
+	HEADMOUNTEDDISPLAY_API EControllerHand GetTrackingSource() const;
 
 	UFUNCTION(BlueprintSetter)
-	void SetTrackingMotionSource(const FName NewSource);
+	HEADMOUNTEDDISPLAY_API void SetTrackingMotionSource(const FName NewSource);
 
-	FName GetTrackingMotionSource();
+	HEADMOUNTEDDISPLAY_API FName GetTrackingMotionSource();
 
 	UFUNCTION(BlueprintSetter)
-	void SetAssociatedPlayerIndex(const int32 NewPlayer);
+	HEADMOUNTEDDISPLAY_API void SetAssociatedPlayerIndex(const int32 NewPlayer);
 
-	void TickComponent(float DeltaTime, enum ELevelTick TickType, FActorComponentTickFunction* ThisTickFunction) override;
-	void BeginDestroy() override;
+	HEADMOUNTEDDISPLAY_API void TickComponent(float DeltaTime, enum ELevelTick TickType, FActorComponentTickFunction* ThisTickFunction) override;
+	HEADMOUNTEDDISPLAY_API void BeginDestroy() override;
 
 	// The following private properties/members are now deprecated and will be removed in later versions.
-	void RefreshDisplayComponent(const bool bForceDestroy = false);
-	void PostLoad() override;
+	HEADMOUNTEDDISPLAY_API void RefreshDisplayComponent(const bool bForceDestroy = false);
+	HEADMOUNTEDDISPLAY_API void PostLoad() override;
 
 	/** Callback for asynchronous display model loads (to set materials, etc.) */
-	void OnDisplayModelLoaded(UPrimitiveComponent* DisplayComponent);
+	HEADMOUNTEDDISPLAY_API void OnDisplayModelLoaded(UPrimitiveComponent* DisplayComponent);
 
 	UPROPERTY(Transient, BlueprintReadOnly, Category = Visualization, meta = (AllowPrivateAccess = "true"))
 	TObjectPtr<UPrimitiveComponent> DisplayComponent;
@@ -121,26 +121,26 @@ class HEADMOUNTEDDISPLAY_API UMotionControllerComponent : public UPrimitiveCompo
 
 public:
 	//~ UObject interface
-	virtual void Serialize(FArchive& Ar) override;
+	HEADMOUNTEDDISPLAY_API virtual void Serialize(FArchive& Ar) override;
 
 #if WITH_EDITOR
-	virtual void PreEditChange(FProperty* PropertyAboutToChange) override;
-	virtual void PostEditChangeProperty(FPropertyChangedEvent& PropertyChangedEvent) override;
+	HEADMOUNTEDDISPLAY_API virtual void PreEditChange(FProperty* PropertyAboutToChange) override;
+	HEADMOUNTEDDISPLAY_API virtual void PostEditChangeProperty(FPropertyChangedEvent& PropertyChangedEvent) override;
 #endif 
 
 	//~ UActorComponent interface
-	virtual void OnRegister() override;
-	virtual void InitializeComponent() override;
-	virtual void OnComponentDestroyed(bool bDestroyingHierarchy) override;
+	HEADMOUNTEDDISPLAY_API virtual void OnRegister() override;
+	HEADMOUNTEDDISPLAY_API virtual void InitializeComponent() override;
+	HEADMOUNTEDDISPLAY_API virtual void OnComponentDestroyed(bool bDestroyingHierarchy) override;
 
 	// Delegate for activation of XRDeviceVisualizationComponent
 	DECLARE_MULTICAST_DELEGATE_OneParam(FActivateVisualizationComponent, bool);
-	static FActivateVisualizationComponent OnActivateVisualizationComponent;
+	static HEADMOUNTEDDISPLAY_API FActivateVisualizationComponent OnActivateVisualizationComponent;
 
 protected:
 	//~ Begin UActorComponent Interface.
-	virtual void CreateRenderState_Concurrent(FRegisterComponentContext* Context) override;
-	virtual void SendRenderTransform_Concurrent() override;
+	HEADMOUNTEDDISPLAY_API virtual void CreateRenderState_Concurrent(FRegisterComponentContext* Context) override;
+	HEADMOUNTEDDISPLAY_API virtual void SendRenderTransform_Concurrent() override;
 	//~ End UActorComponent Interface.
 
 	// Cached Motion Controller that can be read by GetParameterValue. Only valid for the duration of OnMotionControllerUpdated
@@ -148,14 +148,14 @@ protected:
 
 	/** Blueprint Implementable function for responding to updated data from a motion controller (so we can use custom parameter values from it) */
 	UFUNCTION(BlueprintImplementableEvent, Category = "Motion Controller Update")
-	void OnMotionControllerUpdated();
+	HEADMOUNTEDDISPLAY_API void OnMotionControllerUpdated();
 
 	// Returns the value of a custom parameter on the current in use Motion Controller (see member InUseMotionController). Only valid for the duration of OnMotionControllerUpdated 
 	UFUNCTION(BlueprintCallable, Category = "Motion Controller Update")
-	float GetParameterValue(FName InName, bool& bValueFound);
+	HEADMOUNTEDDISPLAY_API float GetParameterValue(FName InName, bool& bValueFound);
 
 	UFUNCTION(BlueprintCallable, Category = "Motion Controller Update")
-	FVector GetHandJointPosition(int jointIndex, bool& bValueFound);
+	HEADMOUNTEDDISPLAY_API FVector GetHandJointPosition(int jointIndex, bool& bValueFound);
 
 private:
 
@@ -166,9 +166,9 @@ private:
 	bool bHasAuthority;
 
 	/** If true, the Position and Orientation args will contain the most recent controller state */
-	bool PollControllerState(FVector& Position, FRotator& Orientation, float WorldToMetersScale);
+	HEADMOUNTEDDISPLAY_API bool PollControllerState(FVector& Position, FRotator& Orientation, float WorldToMetersScale);
 
-	void OnModularFeatureUnregistered(const FName& Type, class IModularFeature* ModularFeature);
+	HEADMOUNTEDDISPLAY_API void OnModularFeatureUnregistered(const FName& Type, class IModularFeature* ModularFeature);
 	IMotionController* PolledMotionController_GameThread;
 	IMotionController* PolledMotionController_RenderThread;
 	FCriticalSection PolledMotionControllerMutex;
@@ -201,4 +201,4 @@ private:
 		FLateUpdateManager LateUpdate;
 	};
 	TSharedPtr< FViewExtension, ESPMode::ThreadSafe > ViewExtension;	
-};
\ No newline at end of file
+};
diff --git a/Engine/Source/Runtime/HeadMountedDisplay/Public/MotionDelayBuffer.h b/Engine/Source/Runtime/HeadMountedDisplay/Public/MotionDelayBuffer.h
--- a/Engine/Source/Runtime/HeadMountedDisplay/Public/MotionDelayBuffer.h
+++ b/Engine/Source/Runtime/HeadMountedDisplay/Public/MotionDelayBuffer.h
@@ -19,41 +19,41 @@ class FMotionDelayClient;
 /* FMotionDelayService
  *****************************************************************************/
 
-class HEADMOUNTEDDISPLAY_API FMotionDelayService
+class FMotionDelayService
 {
 public:
-	static void SetEnabled(bool bEnable);
-	static bool RegisterDelayTarget(USceneComponent* MotionControlledComponent, const int32 PlayerIndex, const FName SourceId);
-	static void RegisterDelayClient(TSharedRef<FMotionDelayClient, ESPMode::ThreadSafe> DelayClient);
+	static HEADMOUNTEDDISPLAY_API void SetEnabled(bool bEnable);
+	static HEADMOUNTEDDISPLAY_API bool RegisterDelayTarget(USceneComponent* MotionControlledComponent, const int32 PlayerIndex, const FName SourceId);
+	static HEADMOUNTEDDISPLAY_API void RegisterDelayClient(TSharedRef<FMotionDelayClient, ESPMode::ThreadSafe> DelayClient);
 };
 
 /* FMotionDelayClient
  *****************************************************************************/
 
-class HEADMOUNTEDDISPLAY_API FMotionDelayClient : public FSceneViewExtensionBase
+class FMotionDelayClient : public FSceneViewExtensionBase
 {
 public:
-	FMotionDelayClient(const FAutoRegister& AutoRegister);
+	HEADMOUNTEDDISPLAY_API FMotionDelayClient(const FAutoRegister& AutoRegister);
 
 	virtual uint32 GetDesiredDelay() const = 0;
 	virtual void GetExemptTargets(TArray<USceneComponent*>& ExemptTargets) const {}
 
-	void Apply_RenderThread(FSceneInterface* Scene);
-	void Restore_RenderThread(FSceneInterface* Scene);
+	HEADMOUNTEDDISPLAY_API void Apply_RenderThread(FSceneInterface* Scene);
+	HEADMOUNTEDDISPLAY_API void Restore_RenderThread(FSceneInterface* Scene);
 
 public:
 	/** ISceneViewExtension interface */
 	virtual void SetupViewFamily(FSceneViewFamily& InViewFamily) override {}
 	virtual void SetupView(FSceneViewFamily& InViewFamily, FSceneView& InView) override {}
-	virtual void BeginRenderViewFamily(FSceneViewFamily& InViewFamily) override;
+	HEADMOUNTEDDISPLAY_API virtual void BeginRenderViewFamily(FSceneViewFamily& InViewFamily) override;
 	virtual void PreRenderView_RenderThread(FRDGBuilder& GraphBuilder, FSceneView& InView) override {}
 	virtual void PreRenderViewFamily_RenderThread(FRDGBuilder& GraphBuilder, FSceneViewFamily& InViewFamily) override {}
-	virtual void PostRenderViewFamily_RenderThread(FRDGBuilder& GraphBuilder, FSceneViewFamily& InViewFamily) override;
-	virtual int32 GetPriority() const override;
-	virtual bool IsActiveThisFrame_Internal(const FSceneViewExtensionContext& Context) const override;
+	HEADMOUNTEDDISPLAY_API virtual void PostRenderViewFamily_RenderThread(FRDGBuilder& GraphBuilder, FSceneViewFamily& InViewFamily) override;
+	HEADMOUNTEDDISPLAY_API virtual int32 GetPriority() const override;
+	HEADMOUNTEDDISPLAY_API virtual bool IsActiveThisFrame_Internal(const FSceneViewExtensionContext& Context) const override;
 
 protected:
-	bool FindDelayTransform(USceneComponent* Target, uint32 Delay, FTransform& TransformOut);
+	HEADMOUNTEDDISPLAY_API bool FindDelayTransform(USceneComponent* Target, uint32 Delay, FTransform& TransformOut);
 
 private:
 	struct FTargetTransform
diff --git a/Engine/Source/Runtime/HeadMountedDisplay/Public/XRGestureConfig.h b/Engine/Source/Runtime/HeadMountedDisplay/Public/XRGestureConfig.h
--- a/Engine/Source/Runtime/HeadMountedDisplay/Public/XRGestureConfig.h
+++ b/Engine/Source/Runtime/HeadMountedDisplay/Public/XRGestureConfig.h
@@ -15,7 +15,7 @@ enum class ESpatialInputGestureAxis : uint8
 
 
 USTRUCT(BlueprintType, meta = (Deprecated = "5.3"))
-struct HEADMOUNTEDDISPLAY_API FXRGestureConfig
+struct FXRGestureConfig
 {
 	GENERATED_BODY()
 
diff --git a/Engine/Source/Runtime/IESFile/Public/IESConverter.h b/Engine/Source/Runtime/IESFile/Public/IESConverter.h
--- a/Engine/Source/Runtime/IESFile/Public/IESConverter.h
+++ b/Engine/Source/Runtime/IESFile/Public/IESConverter.h
@@ -8,22 +8,22 @@
  * To load the IES file image format. IES files exist for many real world lights. The file stores how much light is emitted in a specific direction.
  * The data is usually measured but tools to paint IES files exist.
  */
-class IESFILE_API FIESConverter
+class FIESConverter
 {
 public:
 	/** is loading the file, can take some time, success can be checked with IsValid() */
-	FIESConverter(const uint8* Buffer, uint32 BufferLength);
+	IESFILE_API FIESConverter(const uint8* Buffer, uint32 BufferLength);
 
 	/**
 	 * @return true if the photometric data are valid
 	 * @note Call GetError to get a description of the error
 	 */
-	bool IsValid() const;
+	IESFILE_API bool IsValid() const;
 
 	/**
 	 * @return a brief description of the reason why the IES data is invalid
 	 */
-	const TCHAR* GetError() const;
+	IESFILE_API const TCHAR* GetError() const;
 
 	/**
 	 * @return Multiplier as the texture is normalized
@@ -41,13 +41,13 @@ public:
 		return RawData;
 	}
 
-	uint32 GetWidth() const;
-	uint32 GetHeight() const;
+	IESFILE_API uint32 GetWidth() const;
+	IESFILE_API uint32 GetHeight() const;
 
 	/**
 	 * @return brightness in Lumens
 	 */
-	float GetBrightness() const;
+	IESFILE_API float GetBrightness() const;
 
 private:
 	TArray<uint8> RawData;
diff --git a/Engine/Source/Runtime/ImageCore/Public/ImageCore.h b/Engine/Source/Runtime/ImageCore/Public/ImageCore.h
--- a/Engine/Source/Runtime/ImageCore/Public/ImageCore.h
+++ b/Engine/Source/Runtime/ImageCore/Public/ImageCore.h
@@ -316,7 +316,7 @@ struct FImageView : public FImageInfo
 	IMAGECORE_API void CopyTo(FImage& DestImage, ERawImageFormat::Type DestFormat, EGammaSpace DestGammaSpace) const;
 	
 	// CopyTo same format
-	IMAGECORE_API void CopyTo(FImage& DestImage) const
+	void CopyTo(FImage& DestImage) const
 	{
 		CopyTo(DestImage,Format,GammaSpace);
 	}
@@ -330,7 +330,7 @@ struct FImageView : public FImageInfo
 	}
 
 	// Get one pixel from the image and return in Linear color
-	IMAGECORE_API const FLinearColor GetOnePixelLinear(int32 X,int32 Y,int32 Slice=0) const
+	const FLinearColor GetOnePixelLinear(int32 X,int32 Y,int32 Slice=0) const
 	{
 		void * Ptr = GetPixelPointer(X,Y,Slice);
 		return ERawImageFormat::GetOnePixelLinear(Ptr,Format,GammaSpace);
@@ -479,7 +479,7 @@ public:
 	IMAGECORE_API void CopyTo(FImage& DestImage, ERawImageFormat::Type DestFormat, EGammaSpace DestGammaSpace) const;
 	
 	// CopyTo same format
-	IMAGECORE_API void CopyTo(FImage& DestImage) const
+	void CopyTo(FImage& DestImage) const
 	{
 		CopyTo(DestImage,Format,GammaSpace);
 	}
@@ -568,7 +568,7 @@ public:
 	}
 
 	// Get one pixel from the image and return in Linear color
-	IMAGECORE_API const FLinearColor GetOnePixelLinear(int32 X,int32 Y,int32 Slice=0) const
+	const FLinearColor GetOnePixelLinear(int32 X,int32 Y,int32 Slice=0) const
 	{
 		void * Ptr = GetPixelPointer(X,Y,Slice);
 		return ERawImageFormat::GetOnePixelLinear(Ptr,Format,GammaSpace);
diff --git a/Engine/Source/Runtime/ImageWrapper/Public/Formats/HdrImageWrapper.h b/Engine/Source/Runtime/ImageWrapper/Public/Formats/HdrImageWrapper.h
--- a/Engine/Source/Runtime/ImageWrapper/Public/Formats/HdrImageWrapper.h
+++ b/Engine/Source/Runtime/ImageWrapper/Public/Formats/HdrImageWrapper.h
@@ -14,33 +14,33 @@
 
 /** To load the HDR file image format. Does not support all possible types HDR formats (e.g. xyze is not supported) */
 //  does not use ImageWrapperBase , unlike all other imagewrappers
-class IMAGEWRAPPER_API FHdrImageWrapper : public IImageWrapper
+class FHdrImageWrapper : public IImageWrapper
 {
 public:
 
 	// Todo we should have this for all image wrapper.
-	bool SetCompressedFromView(TArrayView64<const uint8> Data);
+	IMAGEWRAPPER_API bool SetCompressedFromView(TArrayView64<const uint8> Data);
 
 	// IIMageWrapper Interface begin
-	virtual bool SetCompressed(const void* InCompressedData, int64 InCompressedSize) override;
-	virtual bool SetRaw(const void* InRawData, int64 InRawSize, const int32 InWidth, const int32 InHeight, const ERGBFormat InFormat, const int32 InBitDepth, const int32 InBytesPerRow) override;
-	virtual TArray64<uint8> GetCompressed(int32 Quality = 0) override;
-	virtual bool GetRaw(const ERGBFormat InFormat, int32 InBitDepth, TArray64<uint8>& OutRawData) override;
+	IMAGEWRAPPER_API virtual bool SetCompressed(const void* InCompressedData, int64 InCompressedSize) override;
+	IMAGEWRAPPER_API virtual bool SetRaw(const void* InRawData, int64 InRawSize, const int32 InWidth, const int32 InHeight, const ERGBFormat InFormat, const int32 InBitDepth, const int32 InBytesPerRow) override;
+	IMAGEWRAPPER_API virtual TArray64<uint8> GetCompressed(int32 Quality = 0) override;
+	IMAGEWRAPPER_API virtual bool GetRaw(const ERGBFormat InFormat, int32 InBitDepth, TArray64<uint8>& OutRawData) override;
 	
-	virtual bool CanSetRawFormat(const ERGBFormat InFormat, const int32 InBitDepth) const override;
-	virtual ERawImageFormat::Type GetSupportedRawFormat(const ERawImageFormat::Type InFormat) const override;
+	IMAGEWRAPPER_API virtual bool CanSetRawFormat(const ERGBFormat InFormat, const int32 InBitDepth) const override;
+	IMAGEWRAPPER_API virtual ERawImageFormat::Type GetSupportedRawFormat(const ERawImageFormat::Type InFormat) const override;
 
-	virtual int64 GetWidth() const override;
-	virtual int64 GetHeight() const override;
-	virtual int32 GetBitDepth() const override;
-	virtual ERGBFormat GetFormat() const override;
+	IMAGEWRAPPER_API virtual int64 GetWidth() const override;
+	IMAGEWRAPPER_API virtual int64 GetHeight() const override;
+	IMAGEWRAPPER_API virtual int32 GetBitDepth() const override;
+	IMAGEWRAPPER_API virtual ERGBFormat GetFormat() const override;
 
 	// IImageWrapper Interface end
 
 	// GetErrorMessage : nice idea, but not virtual, never called by the standard import path
-	const FText& GetErrorMessage() const;
+	IMAGEWRAPPER_API const FText& GetErrorMessage() const;
 
-	void FreeCompressedData();
+	IMAGEWRAPPER_API void FreeCompressedData();
 
 	using IImageWrapper::GetRaw;
 private:
@@ -78,4 +78,4 @@ private:
 
 	// Reported error
 	FText ErrorMessage;
-};
\ No newline at end of file
+};
diff --git a/Engine/Source/Runtime/ImageWrapper/Public/ImageWrapperHelper.h b/Engine/Source/Runtime/ImageWrapper/Public/ImageWrapperHelper.h
--- a/Engine/Source/Runtime/ImageWrapper/Public/ImageWrapperHelper.h
+++ b/Engine/Source/Runtime/ImageWrapper/Public/ImageWrapperHelper.h
@@ -5,10 +5,10 @@
 #include "CoreMinimal.h"
 #include "IImageWrapper.h"
 
-class IMAGEWRAPPER_API ImageWrapperHelper
+class ImageWrapperHelper
 {
 public:
-	 static FStringView GetFormatExtension(EImageFormat InImageFormat, bool bIncludeDot=false);
-	 static EImageFormat GetImageFormat(FStringView StringExtention);
-	 static const FStringView GetImageFilesFilterString(bool bIncludeAllFiles);
+	 static IMAGEWRAPPER_API FStringView GetFormatExtension(EImageFormat InImageFormat, bool bIncludeDot=false);
+	 static IMAGEWRAPPER_API EImageFormat GetImageFormat(FStringView StringExtention);
+	 static IMAGEWRAPPER_API const FStringView GetImageFilesFilterString(bool bIncludeAllFiles);
 };
diff --git a/Engine/Source/Runtime/ImageWriteQueue/Public/ImageWriteStream.h b/Engine/Source/Runtime/ImageWriteQueue/Public/ImageWriteStream.h
--- a/Engine/Source/Runtime/ImageWriteQueue/Public/ImageWriteStream.h
+++ b/Engine/Source/Runtime/ImageWriteQueue/Public/ImageWriteStream.h
@@ -17,7 +17,7 @@ struct FImageStreamEndpoint;
 /**
  * A pipe that receives image data and forwards it onto 0 or more end points, copying the buffer as few times as possible
  */
-struct IMAGEWRITEQUEUE_API FImagePixelPipe
+struct FImagePixelPipe
 {
 	/**
 	 * Default constructor (an empty pipe)
@@ -38,21 +38,21 @@ struct IMAGEWRITEQUEUE_API FImagePixelPipe
 	 *
 	 * @param InImagePixelData         The data to push through this pipe
 	 */
-	void Push(TUniquePtr<FImagePixelData>&& InImagePixelData);
+	IMAGEWRITEQUEUE_API void Push(TUniquePtr<FImagePixelData>&& InImagePixelData);
 
 	/**
 	 * Add a new end point handler to this pipe.
 	 *
 	 * @param InEndpoint               The new endpoint to add. Potentially used on any thread.
 	 */
-	void AddEndpoint(TUniquePtr<FImageStreamEndpoint>&& InEndpoint);
+	IMAGEWRITEQUEUE_API void AddEndpoint(TUniquePtr<FImageStreamEndpoint>&& InEndpoint);
 
 	/**
 	 * Add a new end point handler to this pipe as a functor.
 	 *
 	 * @param InHandler                A handler function implemented as an anonymous functor. Potentially called on any thread.
 	 */
-	void AddEndpoint(const TFunction<void(TUniquePtr<FImagePixelData>&& )>& InHandler);
+	IMAGEWRITEQUEUE_API void AddEndpoint(const TFunction<void(TUniquePtr<FImagePixelData>&& )>& InHandler);
 
 	/**
 	 * Access this pipe's current set of end points.
@@ -77,7 +77,7 @@ private:
 /**
  * Stream end-point that receives a copy of image data from a thread
  */
-struct IMAGEWRITEQUEUE_API FImageStreamEndpoint
+struct FImageStreamEndpoint
 {
 	virtual ~FImageStreamEndpoint(){}
 
@@ -86,7 +86,7 @@ struct IMAGEWRITEQUEUE_API FImageStreamEndpoint
 	 *
 	 * @param InOwnedImage       Image data to pass through this end point.
 	 */
-	void PipeImage(TUniquePtr<FImagePixelData>&& InOwnedImage);
+	IMAGEWRITEQUEUE_API void PipeImage(TUniquePtr<FImagePixelData>&& InOwnedImage);
 
 private:
 
@@ -94,4 +94,4 @@ private:
 	 * Implemented in derived classes to handle image data being received
 	 */
 	virtual void OnImageReceived(TUniquePtr<FImagePixelData>&& InOwnedImage) {}
-};
\ No newline at end of file
+};
