Fixed render thread check crash.

diff --git a/Engine/Source/Runtime/GeometryFramework/Private/Components/BaseDynamicMeshSceneProxy.cpp b/Engine/Source/Runtime/GeometryFramework/Private/Components/BaseDynamicMeshSceneProxy.cpp
--- a/Engine/Source/Runtime/GeometryFramework/Private/Components/BaseDynamicMeshSceneProxy.cpp
+++ b/Engine/Source/Runtime/GeometryFramework/Private/Components/BaseDynamicMeshSceneProxy.cpp
@@ -28,9 +28,6 @@ FBaseDynamicMeshSceneProxy::FBaseDynamicMeshSceneProxy(UBaseDynamicMeshComponent
 
 FBaseDynamicMeshSceneProxy::~FBaseDynamicMeshSceneProxy()
 {
-	// we are assuming in code below that this is always called from the rendering thread
-	check(IsInRenderingThread());
-
 	// destroy all existing renderbuffers
 	for (FMeshRenderBufferSet* BufferSet : AllocatedBufferSets)
 	{
