fix dll export names somehow winding up in blueprint function display name strings during semi-automated dllexport conversion for UDynamicMesh.h

diff --git a/Engine/Source/Runtime/GeometryFramework/Public/UDynamicMesh.h b/Engine/Source/Runtime/GeometryFramework/Public/UDynamicMesh.h
--- a/Engine/Source/Runtime/GeometryFramework/Public/UDynamicMesh.h
+++ b/Engine/Source/Runtime/GeometryFramework/Public/UDynamicMesh.h
@@ -110,14 +110,14 @@ public:
 	 * This *does not* allocate a new mesh, so any existing mesh pointers/refs are still valid
 	 */
 	UFUNCTION(BlueprintCallable, Category = "Dynamic Mesh")
-	GEOMETRYFRAMEWORK_API UPARAM(DisplayName = "TarGEOMETRYFRAMEWORK_API get") UDynamicMesh* Reset();
+	GEOMETRYFRAMEWORK_API UPARAM(DisplayName = "Target") UDynamicMesh* Reset();
 
 	/**
 	 * Clear the internal mesh to a 100x100x100 cube with base at the origin.
 	 * This this instead of Reset() if an initially-empty mesh is undesirable (eg for a Component)
 	 */
 	UFUNCTION(BlueprintCallable, Category = "Dynamic Mesh")
-	GEOMETRYFRAMEWORK_API UPARAM(DisplayName = "TarGEOMETRYFRAMEWORK_API get") UDynamicMesh* ResetToCube();
+	GEOMETRYFRAMEWORK_API UPARAM(DisplayName = "Target") UDynamicMesh* ResetToCube();
 
 	//
 	// Native access/modification functions
@@ -138,7 +138,7 @@ public:
 	 * @return number of triangles in the mesh
 	 */
 	UFUNCTION(BlueprintCallable, Category = "Dynamic Mesh")
-	GEOMETRYFRAMEWORK_API UPARAM(DisplayName = "TriGEOMETRYFRAMEWORK_API angle Count") int32 GetTriangleCount() const;
+	GEOMETRYFRAMEWORK_API UPARAM(DisplayName = "Triangle Count") int32 GetTriangleCount() const;
 
 
 	/** 
