Change the default TangentsType for all new DynamicMeshComponent and expose TangentsType to the UI

Make TangentsType have an explicit 'Default' value, which will currently map to ExternallyProvided. ExternallyProvided is also renamed to "From Dynamic Mesh" in the UI.

diff --git a/Engine/Source/Runtime/Core/Public/UObject/UE5ReleaseStreamObjectVersion.h b/Engine/Source/Runtime/Core/Public/UObject/UE5ReleaseStreamObjectVersion.h
--- a/Engine/Source/Runtime/Core/Public/UObject/UE5ReleaseStreamObjectVersion.h
+++ b/Engine/Source/Runtime/Core/Public/UObject/UE5ReleaseStreamObjectVersion.h
@@ -135,6 +135,9 @@ struct FUE5ReleaseStreamObjectVersion
 
 		// Added member reference to linked anim graphs
 		LinkedAnimGraphMemberReference,
+
+		// Changed default tangent behavior for new dynamic mesh components
+		DynamicMeshComponentsDefaultUseExternalTangents,
 		
 		// -----<new versions can be added above this line>-------------------------------------------------
 		VersionPlusOne,
diff --git a/Engine/Source/Runtime/GeometryFramework/Private/Components/DynamicMeshComponent.cpp b/Engine/Source/Runtime/GeometryFramework/Private/Components/DynamicMeshComponent.cpp
--- a/Engine/Source/Runtime/GeometryFramework/Private/Components/DynamicMeshComponent.cpp
+++ b/Engine/Source/Runtime/GeometryFramework/Private/Components/DynamicMeshComponent.cpp
@@ -18,6 +18,8 @@
 #include "Changes/MeshChange.h"
 #include "DynamicMesh/MeshTransforms.h"
 
+#include "UObject/UE5ReleaseStreamObjectVersion.h"
+
 // default proxy for this component
 #include "Components/DynamicMeshSceneProxy.h"
 
@@ -80,11 +82,26 @@ UDynamicMeshComponent::UDynamicMeshComponent(const FObjectInitializer& ObjectIni
 	ResetProxy();
 }
 
+void UDynamicMeshComponent::Serialize(FArchive& Ar)
+{
+	Super::Serialize(Ar);
+	Ar.UsingCustomVersion(FUE5ReleaseStreamObjectVersion::GUID);
+}
 
 void UDynamicMeshComponent::PostLoad()
 {
 	Super::PostLoad();
 
+	const int32 UE5ReleaseStreamObjectVersion = GetLinkerCustomVersion(FUE5ReleaseStreamObjectVersion::GUID);
+	if (UE5ReleaseStreamObjectVersion < FUE5ReleaseStreamObjectVersion::DynamicMeshComponentsDefaultUseExternalTangents)
+	{
+		// Set the old default value
+		if (TangentsType == EDynamicMeshComponentTangentsMode::Default)
+		{
+			TangentsType = EDynamicMeshComponentTangentsMode::NoTangents;
+		}
+	}
+
 	// The intention here is that MeshObject is never nullptr, however we cannot guarantee this as a subclass
 	// may have set it to null, and/or some type of serialization issue has caused it to fail to save/load.
 	// Avoid immediate crashes by creating a new UDynamicMesh here in such cases
@@ -115,7 +132,11 @@ void UDynamicMeshComponent::PostEditChangeProperty(FPropertyChangedEvent& Proper
 	Super::PostEditChangeProperty(PropertyChangedEvent);
 
 	const FName PropName = PropertyChangedEvent.GetPropertyName();
-	if ( (PropName == GET_MEMBER_NAME_CHECKED(UDynamicMeshComponent, bEnableComplexCollision)) ||
+	if (PropName == GET_MEMBER_NAME_CHECKED(UDynamicMeshComponent, TangentsType))
+	{
+		InvalidateAutoCalculatedTangents();
+	}
+	else if ( (PropName == GET_MEMBER_NAME_CHECKED(UDynamicMeshComponent, bEnableComplexCollision)) ||
 		(PropName == GET_MEMBER_NAME_CHECKED(UDynamicMeshComponent, CollisionType)) ||
 		(PropName == GET_MEMBER_NAME_CHECKED(UDynamicMeshComponent, bDeferCollisionUpdates))  )
 	{
@@ -278,7 +299,7 @@ void UDynamicMeshComponent::InvalidateAutoCalculatedTangents()
 
 const UE::Geometry::FMeshTangentsf* UDynamicMeshComponent::GetAutoCalculatedTangents() 
 { 
-	if (TangentsType == EDynamicMeshComponentTangentsMode::AutoCalculated && GetDynamicMesh()->GetMeshRef().HasAttributes())
+	if (GetTangentsType() == EDynamicMeshComponentTangentsMode::AutoCalculated && GetDynamicMesh()->GetMeshRef().HasAttributes())
 	{
 		UpdateAutoCalculatedTangents();
 		return (bAutoCalculatedTangentsValid) ? &AutoCalculatedTangents : nullptr;
@@ -288,7 +309,7 @@ const UE::Geometry::FMeshTangentsf* UDynamicMeshComponent::GetAutoCalculatedTang
 
 void UDynamicMeshComponent::UpdateAutoCalculatedTangents()
 {
-	if (TangentsType == EDynamicMeshComponentTangentsMode::AutoCalculated && bAutoCalculatedTangentsValid == false)
+	if (GetTangentsType() == EDynamicMeshComponentTangentsMode::AutoCalculated && bAutoCalculatedTangentsValid == false)
 	{
 		GetDynamicMesh()->ProcessMesh([&](const FDynamicMesh3& Mesh)
 		{
diff --git a/Engine/Source/Runtime/GeometryFramework/Public/Components/BaseDynamicMeshComponent.h b/Engine/Source/Runtime/GeometryFramework/Public/Components/BaseDynamicMeshComponent.h
--- a/Engine/Source/Runtime/GeometryFramework/Public/Components/BaseDynamicMeshComponent.h
+++ b/Engine/Source/Runtime/GeometryFramework/Public/Components/BaseDynamicMeshComponent.h
@@ -51,7 +51,9 @@ enum class EDynamicMeshComponentTangentsMode : uint8
 	/** Tangents will be automatically calculated on demand. Note that mesh changes due to tangents calculation will *not* be broadcast via MeshChange events! */
 	AutoCalculated,
 	/** Tangents are externally provided via the FDynamicMesh3 AttributeSet */
-	ExternallyProvided
+	ExternallyProvided UMETA(DisplayName = "From Dynamic Mesh"),
+	/** Tangents mode will be set to the most commonly-useful default -- currently "From Dynamic Mesh" */
+	Default = 255
 };
 
 
diff --git a/Engine/Source/Runtime/GeometryFramework/Public/Components/DynamicMeshComponent.h b/Engine/Source/Runtime/GeometryFramework/Public/Components/DynamicMeshComponent.h
--- a/Engine/Source/Runtime/GeometryFramework/Public/Components/DynamicMeshComponent.h
+++ b/Engine/Source/Runtime/GeometryFramework/Public/Components/DynamicMeshComponent.h
@@ -469,21 +469,39 @@ public:
 
 
 	//===============================================================================================================
-	// Triangle-Vertex Tangents support. The default behavior is to not use Tangents, this will lead to incorrect
-	// rendering for any material with Normal Maps and some other shaders.
+	// Triangle-Vertex Tangents support. The default behavior is to use the provided external Tangents.
 	// If TangentsType == EDynamicMeshComponentTangentsMode::ExternallyProvided, the Tangent and Bitangent attributes of
 	//   the FDynamicMesh3 AttributeSet are used at the SceneProxy level, the Component is not involved
 	// If TangentsType == EDynamicMeshComponentTangentsMode::AutoCalculated, the Tangents are computed internally using
 	//   a fast MikkT approximation via FMeshTangentsf. They will be recomputed when the mesh is modified, however
 	//   they are *not* recomputed when using the Fast Update functions above (in that case InvalidateAutoCalculatedTangents() 
 	//   can be used to force recomputation)
+	// If TangentsType == EDynamicMeshComponentTangentsMode::NoTangents, the Component will not use Tangents, which will
+	//   lead to incorrect rendering for any material with Normal Maps and some other shaders.
 	//
+private:
+	// Default what the 'Default' tangent type should map to
+	static constexpr EDynamicMeshComponentTangentsMode UseDefaultTangentsType = EDynamicMeshComponentTangentsMode::ExternallyProvided;
+
 public:
 	UFUNCTION(BlueprintCallable, Category = "Dynamic Mesh Component")
 	GEOMETRYFRAMEWORK_API void SetTangentsType(EDynamicMeshComponentTangentsMode NewTangentsType);
 
 	UFUNCTION(BlueprintCallable, Category = "Dynamic Mesh Component")
-	EDynamicMeshComponentTangentsMode GetTangentsType() const { return TangentsType; }
+	EDynamicMeshComponentTangentsMode GetTangentsType() const
+	{
+		return TangentsType == EDynamicMeshComponentTangentsMode::Default ? UseDefaultTangentsType : TangentsType;
+	}
+
+private:
+	// pure version of GetTangentsType, so it can be used as a getter below (getters must be BlueprintPure)
+	UFUNCTION(BlueprintCallable, BlueprintPure, BlueprintInternalUseOnly, Category = "Dynamic Mesh Component")
+	EDynamicMeshComponentTangentsMode GetTangentsTypePure() const
+	{
+		return GetTangentsType();
+	}
+
+public:
 
 	/** This function marks the auto tangents as dirty, they will be recomputed before they are used again */
 	GEOMETRYFRAMEWORK_API virtual void InvalidateAutoCalculatedTangents();
@@ -492,9 +510,9 @@ public:
 	GEOMETRYFRAMEWORK_API const UE::Geometry::FMeshTangentsf* GetAutoCalculatedTangents();
 
 protected:
-	/** How should Tangents be calculated/handled */
-	UPROPERTY()
-	EDynamicMeshComponentTangentsMode TangentsType = EDynamicMeshComponentTangentsMode::NoTangents;
+	/** Tangent source defines whether we use the provided tangents on the Dynamic Mesh, autogenerate tangents, or do not use tangents */
+	UPROPERTY(EditAnywhere, BlueprintReadWrite, BlueprintSetter=SetTangentsType, BlueprintGetter=GetTangentsTypePure, Category = "Dynamic Mesh Component|Rendering")
+	EDynamicMeshComponentTangentsMode TangentsType = EDynamicMeshComponentTangentsMode::Default;
 
 	/** true if AutoCalculatedTangents has been computed for current mesh */
 	bool bAutoCalculatedTangentsValid = false;
@@ -685,6 +703,7 @@ protected:
 	GEOMETRYFRAMEWORK_API virtual void OnChildDetached(USceneComponent* ChildComponent) override;
 
 	//~ UObject Interface.
+	GEOMETRYFRAMEWORK_API virtual void Serialize(FArchive& Ar) override;
 	GEOMETRYFRAMEWORK_API virtual void PostLoad() override;
 	GEOMETRYFRAMEWORK_API virtual void BeginDestroy() override;
 #if WITH_EDITOR
