make Dynamic Mesh Components show vertex colors when the corresponding ShowFlag is set

diff --git a/Engine/Source/Runtime/GeometryFramework/Private/Components/BaseDynamicMeshSceneProxy.cpp b/Engine/Source/Runtime/GeometryFramework/Private/Components/BaseDynamicMeshSceneProxy.cpp
--- a/Engine/Source/Runtime/GeometryFramework/Private/Components/BaseDynamicMeshSceneProxy.cpp
+++ b/Engine/Source/Runtime/GeometryFramework/Private/Components/BaseDynamicMeshSceneProxy.cpp
@@ -95,13 +95,63 @@ void FBaseDynamicMeshSceneProxy::UpdatedReferencedMaterials()
 #endif
 }
 
+FMaterialRenderProxy* FBaseDynamicMeshSceneProxy::GetEngineVertexColorMaterialProxy(FMeshElementCollector& Collector, const FEngineShowFlags& EngineShowFlags, bool bProxyIsSelected, bool bIsHovered)
+{
+	FMaterialRenderProxy* ForceOverrideMaterialProxy = nullptr;
+#if UE_ENABLE_DEBUG_DRAWING
+	if (bProxyIsSelected && EngineShowFlags.VertexColors && AllowDebugViewmodes())
+	{
+		// Override the mesh's material with our material that draws the vertex colors
+		UMaterial* VertexColorVisualizationMaterial = NULL;
+		switch (GVertexColorViewMode)
+		{
+		case EVertexColorViewMode::Color:
+			VertexColorVisualizationMaterial = GEngine->VertexColorViewModeMaterial_ColorOnly;
+			break;
+
+		case EVertexColorViewMode::Alpha:
+			VertexColorVisualizationMaterial = GEngine->VertexColorViewModeMaterial_AlphaAsColor;
+			break;
+
+		case EVertexColorViewMode::Red:
+			VertexColorVisualizationMaterial = GEngine->VertexColorViewModeMaterial_RedOnly;
+			break;
+
+		case EVertexColorViewMode::Green:
+			VertexColorVisualizationMaterial = GEngine->VertexColorViewModeMaterial_GreenOnly;
+			break;
+
+		case EVertexColorViewMode::Blue:
+			VertexColorVisualizationMaterial = GEngine->VertexColorViewModeMaterial_BlueOnly;
+			break;
+		}
+		check(VertexColorVisualizationMaterial != NULL);
+
+		// Note: static mesh renderer does something more complicated involving per-section selection,
+		// but whole component selection seems ok for now
+		bool bSectionIsSelected = bProxyIsSelected;
+
+		auto VertexColorVisualizationMaterialInstance = new FColoredMaterialRenderProxy(
+			VertexColorVisualizationMaterial->GetRenderProxy(),
+			GetSelectionColor(FLinearColor::White, bSectionIsSelected, bIsHovered)
+		);
+
+		Collector.RegisterOneFrameMaterialProxy(VertexColorVisualizationMaterialInstance);
+		ForceOverrideMaterialProxy = VertexColorVisualizationMaterialInstance;
+	}
+#endif
+	return ForceOverrideMaterialProxy;
+}
+
 void FBaseDynamicMeshSceneProxy::GetDynamicMeshElements(const TArray<const FSceneView*>& Views, const FSceneViewFamily& ViewFamily, uint32 VisibilityMap, FMeshElementCollector& Collector) const
 {
 	QUICK_SCOPE_CYCLE_COUNTER(STAT_BaseDynamicMeshSceneProxy_GetDynamicMeshElements);
 
-	bool bIsWireframeViewMode = (AllowDebugViewmodes() && ViewFamily.EngineShowFlags.Wireframe);
+	const FEngineShowFlags& EngineShowFlags = ViewFamily.EngineShowFlags;
+	bool bIsWireframeViewMode = (AllowDebugViewmodes() && EngineShowFlags.Wireframe);
 	bool bWantWireframeOnShaded = ParentBaseComponent->GetEnableWireframeRenderPass();
 	bool bWireframe = bIsWireframeViewMode || bWantWireframeOnShaded;
+	const bool bProxyIsSelected = IsSelected();
 
 	// Get wireframe material proxy if requested and available, otherwise disable wireframe
 	FMaterialRenderProxy* WireframeMaterialProxy = nullptr;
@@ -110,7 +160,7 @@ void FBaseDynamicMeshSceneProxy::GetDynamicMeshElements(const TArray<const FScen
 		UMaterialInterface* WireframeMaterial = UBaseDynamicMeshComponent::GetDefaultWireframeMaterial_RenderThread();
 		if (WireframeMaterial != nullptr)
 		{
-			FLinearColor UseWireframeColor = (IsSelected() && (bWantWireframeOnShaded == false || bIsWireframeViewMode))  ?
+			FLinearColor UseWireframeColor = (bProxyIsSelected && (bWantWireframeOnShaded == false || bIsWireframeViewMode))  ?
 				GEngine->GetSelectedMaterialColor() : ParentBaseComponent->WireframeColor;
 			FColoredMaterialRenderProxy* WireframeMaterialInstance = new FColoredMaterialRenderProxy(
 				WireframeMaterial->GetRenderProxy(), UseWireframeColor);
@@ -123,18 +173,17 @@ void FBaseDynamicMeshSceneProxy::GetDynamicMeshElements(const TArray<const FScen
 		}
 	}
 
-	// will use this material instead of any others below, if it becomes non-null
-	UMaterialInterface* ForceOverrideMaterial = nullptr;
-
-	// Note: there is an engine show flags for vertex colors. However if we enable that, then
-	// we will fail the used-materials checks in the Editor, because UBaseDynamicMeshComponent::GetUsedMaterials()
-	// cannot check the active show flags. One option would be to always add the vertex color material there...
-	const bool bVertexColor = ParentBaseComponent->ColorMode == EDynamicMeshComponentColorOverrideMode::VertexColors ||
-								ParentBaseComponent->ColorMode == EDynamicMeshComponentColorOverrideMode::Polygroups ||
-								ParentBaseComponent->ColorMode == EDynamicMeshComponentColorOverrideMode::Constant;
-	if (bVertexColor)
+	FMaterialRenderProxy* ForceOverrideMaterialProxy = GetEngineVertexColorMaterialProxy(Collector, EngineShowFlags, bProxyIsSelected, IsHovered());
+	// If engine show flags aren't setting vertex color, also check if the component requested custom vertex color modes for the dynamic mesh
+	if (!ForceOverrideMaterialProxy)
 	{
-		ForceOverrideMaterial = UBaseDynamicMeshComponent::GetDefaultVertexColorMaterial_RenderThread();
+		const bool bVertexColor = ParentBaseComponent->ColorMode == EDynamicMeshComponentColorOverrideMode::VertexColors ||
+			ParentBaseComponent->ColorMode == EDynamicMeshComponentColorOverrideMode::Polygroups ||
+			ParentBaseComponent->ColorMode == EDynamicMeshComponentColorOverrideMode::Constant;
+		if (bVertexColor)
+		{
+			ForceOverrideMaterialProxy = UBaseDynamicMeshComponent::GetDefaultVertexColorMaterial_RenderThread()->GetRenderProxy();
+		}
 	}
 
 	ESceneDepthPriorityGroup DepthPriority = SDPG_World;
@@ -142,10 +191,10 @@ void FBaseDynamicMeshSceneProxy::GetDynamicMeshElements(const TArray<const FScen
 	TArray<FMeshRenderBufferSet*> Buffers;
 	GetActiveRenderBufferSets(Buffers);
 
-	UMaterialInterface* UseSecondaryMaterial = ForceOverrideMaterial;
-	if (ParentBaseComponent->HasSecondaryRenderMaterial() && ForceOverrideMaterial == nullptr)
+	FMaterialRenderProxy* SecondaryMaterialProxy = ForceOverrideMaterialProxy;
+	if (ParentBaseComponent->HasSecondaryRenderMaterial() && ForceOverrideMaterialProxy == nullptr)
 	{
-		UseSecondaryMaterial = ParentBaseComponent->GetSecondaryRenderMaterial();
+		SecondaryMaterialProxy = ParentBaseComponent->GetSecondaryRenderMaterial()->GetRenderProxy();
 	}
 	bool bDrawSecondaryBuffers = ParentBaseComponent->GetSecondaryBuffersVisibility();
 
@@ -165,16 +214,16 @@ void FBaseDynamicMeshSceneProxy::GetDynamicMeshElements(const TArray<const FScen
 			// Draw the mesh.
 			for (FMeshRenderBufferSet* BufferSet : Buffers)
 			{
-				UMaterialInterface* UseMaterial = BufferSet->Material;
-				if (ParentBaseComponent->HasOverrideRenderMaterial(0))
+				FMaterialRenderProxy* MaterialProxy = ForceOverrideMaterialProxy;
+				if (!MaterialProxy)
 				{
-					UseMaterial = ParentBaseComponent->GetOverrideRenderMaterial(0);
-				}
-				if (ForceOverrideMaterial)
-				{
-					UseMaterial = ForceOverrideMaterial;
+					UMaterialInterface* UseMaterial = BufferSet->Material;
+					if (ParentBaseComponent->HasOverrideRenderMaterial(0))
+					{
+						UseMaterial = ParentBaseComponent->GetOverrideRenderMaterial(0);
+					}
+					MaterialProxy = UseMaterial->GetRenderProxy();
 				}
-				FMaterialRenderProxy* MaterialProxy = UseMaterial->GetRenderProxy();
 
 				if (BufferSet->TriangleCount == 0)
 				{
@@ -209,7 +258,7 @@ void FBaseDynamicMeshSceneProxy::GetDynamicMeshElements(const TArray<const FScen
 				}
 
 				// draw secondary buffer if we have it, falling back to base material if we don't have the Secondary material
-				FMaterialRenderProxy* UseSecondaryMaterialProxy = (UseSecondaryMaterial != nullptr) ? UseSecondaryMaterial->GetRenderProxy() : MaterialProxy;
+				FMaterialRenderProxy* UseSecondaryMaterialProxy = (SecondaryMaterialProxy != nullptr) ? SecondaryMaterialProxy : MaterialProxy;
 				if (bDrawSecondaryBuffers && BufferSet->SecondaryIndexBuffer.Indices.Num() > 0 && UseSecondaryMaterialProxy != nullptr)
 				{
 					if (bWireframe)
diff --git a/Engine/Source/Runtime/GeometryFramework/Public/Components/BaseDynamicMeshSceneProxy.h b/Engine/Source/Runtime/GeometryFramework/Public/Components/BaseDynamicMeshSceneProxy.h
--- a/Engine/Source/Runtime/GeometryFramework/Public/Components/BaseDynamicMeshSceneProxy.h
+++ b/Engine/Source/Runtime/GeometryFramework/Public/Components/BaseDynamicMeshSceneProxy.h
@@ -948,6 +948,10 @@ public:
 	// FBaseDynamicMeshSceneProxy implementation
 	//
 
+	/**
+	 * If EngineShowFlags request vertex color rendering, returns the appropriate vertex color override material's render proxy.  Otherwise returns nullptr.
+	 */
+	GEOMETRYFRAMEWORK_API static FMaterialRenderProxy* GetEngineVertexColorMaterialProxy(FMeshElementCollector& Collector, const FEngineShowFlags& EngineShowFlags, bool bProxyIsSelected, bool bIsHovered);
 
 	/**
 	 * Render set of active RenderBuffers returned by GetActiveRenderBufferSets
