Fix ensure when ExternallyProvided tangents are requested on a DynamicMeshComponent but the DynamicMesh does not have tangents. Fall back to the no-tangents case.

diff --git a/Engine/Source/Runtime/GeometryFramework/Private/Components/DynamicMeshSceneProxy.h b/Engine/Source/Runtime/GeometryFramework/Private/Components/DynamicMeshSceneProxy.h
--- a/Engine/Source/Runtime/GeometryFramework/Private/Components/DynamicMeshSceneProxy.h
+++ b/Engine/Source/Runtime/GeometryFramework/Private/Components/DynamicMeshSceneProxy.h
@@ -94,11 +94,16 @@ public:
 		EDynamicMeshComponentTangentsMode TangentsType = ParentComponent->GetTangentsType();
 		if (TangentsType == EDynamicMeshComponentTangentsMode::ExternallyProvided)
 		{
-			UE::Geometry::FDynamicMeshTangents Tangents(ParentComponent->GetRenderMesh());
-			return [Tangents](int VertexID, int TriangleID, int TriVtxIdx, const FVector3f& Normal, FVector3f& TangentX, FVector3f& TangentY) -> void
+			UE::Geometry::FDynamicMesh3* RenderMesh = ParentComponent->GetRenderMesh();
+			// If the RenderMesh has tangents, use them. Otherwise we fall back to the orthogonal basis, below.
+			if (RenderMesh && RenderMesh->HasAttributes() && RenderMesh->Attributes()->HasTangentSpace())
 			{
-				Tangents.GetTangentVectors(TriangleID, TriVtxIdx, Normal, TangentX, TangentY);
-			};
+				UE::Geometry::FDynamicMeshTangents Tangents(RenderMesh);
+				return [Tangents](int VertexID, int TriangleID, int TriVtxIdx, const FVector3f& Normal, FVector3f& TangentX, FVector3f& TangentY) -> void
+				{
+					Tangents.GetTangentVectors(TriangleID, TriVtxIdx, Normal, TangentX, TangentY);
+				};
+			}
 		}
 		else if (TangentsType == EDynamicMeshComponentTangentsMode::AutoCalculated && bSkipAutoCompute == false)
 		{
