Fixed ensure due to enqueuing render command in scene proxy destructor off the render thread.

diff --git a/Engine/Source/Runtime/GeometryFramework/Public/Components/BaseDynamicMeshSceneProxy.h b/Engine/Source/Runtime/GeometryFramework/Public/Components/BaseDynamicMeshSceneProxy.h
--- a/Engine/Source/Runtime/GeometryFramework/Public/Components/BaseDynamicMeshSceneProxy.h
+++ b/Engine/Source/Runtime/GeometryFramework/Public/Components/BaseDynamicMeshSceneProxy.h
@@ -98,8 +98,6 @@ public:
 
 	virtual ~FMeshRenderBufferSet()
 	{
-		check(IsInRenderingThread());
-
 		if (TriangleCount > 0)
 		{
 			PositionVertexBuffer.ReleaseResource();
@@ -393,11 +391,7 @@ protected:
 			return;
 		}
 
-		ENQUEUE_RENDER_COMMAND(FMeshRenderBufferSetDestroy)(
-			[BufferSet](FRHICommandListImmediate& RHICmdList)
-		{
-			delete BufferSet;
-		});
+		delete BufferSet;
 	}
 
 
