Add support for collision rendering showflags to BaseDynamicMeshSceneProxy

diff --git a/Engine/Source/Runtime/GeometryFramework/Private/Components/BaseDynamicMeshSceneProxy.cpp b/Engine/Source/Runtime/GeometryFramework/Private/Components/BaseDynamicMeshSceneProxy.cpp
--- a/Engine/Source/Runtime/GeometryFramework/Private/Components/BaseDynamicMeshSceneProxy.cpp
+++ b/Engine/Source/Runtime/GeometryFramework/Private/Components/BaseDynamicMeshSceneProxy.cpp
@@ -4,6 +4,7 @@
 #include "MaterialDomain.h"
 #include "Materials/Material.h"
 #include "Materials/MaterialRenderProxy.h"
+#include "PhysicsEngine/BodySetup.h"
 #include "RayTracingDefinitions.h"
 #include "RayTracingInstance.h"
 #include "SceneInterface.h"
@@ -24,6 +25,8 @@ FBaseDynamicMeshSceneProxy::FBaseDynamicMeshSceneProxy(UBaseDynamicMeshComponent
 	}
 
 	bUsePerTriangleNormals = Component->GetFlatShadingEnabled();
+	
+	SetCollisionData();
 }
 
 FBaseDynamicMeshSceneProxy::~FBaseDynamicMeshSceneProxy()
@@ -131,7 +134,7 @@ FMaterialRenderProxy* FBaseDynamicMeshSceneProxy::GetEngineVertexColorMaterialPr
 		// but whole component selection seems ok for now
 		bool bSectionIsSelected = bProxyIsSelected;
 
-		auto VertexColorVisualizationMaterialInstance = new FColoredMaterialRenderProxy(
+		FColoredMaterialRenderProxy* VertexColorVisualizationMaterialInstance = new FColoredMaterialRenderProxy(
 			VertexColorVisualizationMaterial->GetRenderProxy(),
 			GetSelectionColor(FLinearColor::White, bSectionIsSelected, bIsHovered)
 		);
@@ -143,6 +146,31 @@ FMaterialRenderProxy* FBaseDynamicMeshSceneProxy::GetEngineVertexColorMaterialPr
 	return ForceOverrideMaterialProxy;
 }
 
+bool FBaseDynamicMeshSceneProxy::IsCollisionView(const FEngineShowFlags& EngineShowFlags, bool& bDrawSimpleCollision, bool& bDrawComplexCollision) const
+{
+	bDrawSimpleCollision = bDrawComplexCollision = false;
+
+	bool bDrawCollisionView = (EngineShowFlags.CollisionVisibility || EngineShowFlags.CollisionPawn);
+
+#if UE_ENABLE_DEBUG_DRAWING
+	// If in a 'collision view' and collision is enabled
+	if (bHasCollisionData && bDrawCollisionView && IsCollisionEnabled())
+	{
+		// See if we have a response to the interested channel
+		bool bHasResponse = EngineShowFlags.CollisionPawn && CollisionResponse.GetResponse(ECC_Pawn) != ECR_Ignore;
+		bHasResponse |= EngineShowFlags.CollisionVisibility && CollisionResponse.GetResponse(ECC_Visibility) != ECR_Ignore;
+
+		if(bHasResponse)
+		{
+			// Visibility uses complex and pawn uses simple. However, if UseSimpleAsComplex or UseComplexAsSimple is used we need to adjust accordingly
+			bDrawComplexCollision = (EngineShowFlags.CollisionVisibility && CollisionTraceFlag != ECollisionTraceFlag::CTF_UseSimpleAsComplex) || (EngineShowFlags.CollisionPawn && CollisionTraceFlag == ECollisionTraceFlag::CTF_UseComplexAsSimple);
+			bDrawSimpleCollision  = (EngineShowFlags.CollisionPawn && CollisionTraceFlag != ECollisionTraceFlag::CTF_UseComplexAsSimple) || (EngineShowFlags.CollisionVisibility && CollisionTraceFlag == ECollisionTraceFlag::CTF_UseSimpleAsComplex);
+		}
+	}
+#endif
+	return bDrawCollisionView;
+}
+
 void FBaseDynamicMeshSceneProxy::GetDynamicMeshElements(const TArray<const FSceneView*>& Views, const FSceneViewFamily& ViewFamily, uint32 VisibilityMap, FMeshElementCollector& Collector) const
 {
 	QUICK_SCOPE_CYCLE_COUNTER(STAT_BaseDynamicMeshSceneProxy_GetDynamicMeshElements);
@@ -153,6 +181,22 @@ void FBaseDynamicMeshSceneProxy::GetDynamicMeshElements(const TArray<const FScen
 	bool bWireframe = bIsWireframeViewMode || bWantWireframeOnShaded;
 	const bool bProxyIsSelected = IsSelected();
 
+
+	TArray<FMeshRenderBufferSet*> Buffers;
+	GetActiveRenderBufferSets(Buffers);
+
+#if UE_ENABLE_DEBUG_DRAWING
+	bool bDrawSimpleCollision = false, bDrawComplexCollision = false;
+	const bool bDrawCollisionView = IsCollisionView(EngineShowFlags, bDrawSimpleCollision, bDrawComplexCollision);
+
+	// If we're in a collision view, run the only draw the collision and return without drawing mesh normally
+	if (bDrawCollisionView)
+	{
+		GetCollisionDynamicMeshElements(Buffers, EngineShowFlags, bDrawCollisionView, bDrawSimpleCollision, bDrawComplexCollision, bProxyIsSelected, Views, VisibilityMap, Collector);
+		return;
+	}
+#endif
+
 	// Get wireframe material proxy if requested and available, otherwise disable wireframe
 	FMaterialRenderProxy* WireframeMaterialProxy = nullptr;
 	if (bWireframe)
@@ -188,8 +232,6 @@ void FBaseDynamicMeshSceneProxy::GetDynamicMeshElements(const TArray<const FScen
 
 	ESceneDepthPriorityGroup DepthPriority = SDPG_World;
 
-	TArray<FMeshRenderBufferSet*> Buffers;
-	GetActiveRenderBufferSets(Buffers);
 
 	FMaterialRenderProxy* SecondaryMaterialProxy = ForceOverrideMaterialProxy;
 	if (ParentBaseComponent->HasSecondaryRenderMaterial() && ForceOverrideMaterialProxy == nullptr)
@@ -277,6 +319,131 @@ void FBaseDynamicMeshSceneProxy::GetDynamicMeshElements(const TArray<const FScen
 			}
 		}
 	}
+
+#if UE_ENABLE_DEBUG_DRAWING
+	GetCollisionDynamicMeshElements(Buffers, EngineShowFlags, bDrawCollisionView, bDrawSimpleCollision, bDrawComplexCollision, bProxyIsSelected, Views, VisibilityMap, Collector);
+#endif
+}
+
+void FBaseDynamicMeshSceneProxy::GetCollisionDynamicMeshElements(TArray<FMeshRenderBufferSet*>& Buffers, 
+	const FEngineShowFlags& EngineShowFlags, bool bDrawCollisionView, bool bDrawSimpleCollision, bool bDrawComplexCollision, bool bProxyIsSelected,
+	const TArray<const FSceneView*>& Views, uint32 VisibilityMap, FMeshElementCollector& Collector) const
+{
+#if UE_ENABLE_DEBUG_DRAWING
+	if (!bHasCollisionData)
+	{
+		return;
+	}
+
+	// Note: This is closely following StaticMeshRendering.cpp's collision rendering code, from its GetDynamicMeshElements() implementation
+	FColor SimpleCollisionColor = FColor(157, 149, 223, 255);
+	FColor ComplexCollisionColor = FColor(0, 255, 255, 255);
+
+	for (int32 ViewIndex = 0; ViewIndex < Views.Num(); ViewIndex++)
+	{
+		if (VisibilityMap & (1 << ViewIndex))
+		{
+			const FSceneView* View = Views[ViewIndex];
+
+			if(AllowDebugViewmodes())
+			{
+				// Should we draw the mesh wireframe to indicate we are using the mesh as collision
+				bool bDrawComplexWireframeCollision = (EngineShowFlags.Collision && IsCollisionEnabled() && CollisionTraceFlag == ECollisionTraceFlag::CTF_UseComplexAsSimple);
+
+				// If drawing complex collision as solid or wireframe
+				if(bDrawComplexWireframeCollision || (bDrawCollisionView && bDrawComplexCollision))
+				{
+					bool bDrawWireframe = !bDrawCollisionView;
+
+					UMaterial* MaterialToUse = UMaterial::GetDefaultMaterial(MD_Surface);
+					FLinearColor DrawCollisionColor = GetWireframeColor();
+					// Collision view modes draw collision mesh as solid
+					if(bDrawCollisionView)
+					{
+						MaterialToUse = GEngine->ShadedLevelColorationUnlitMaterial;
+					}
+					// Wireframe, choose color based on complex or simple
+					else
+					{
+						MaterialToUse = GEngine->WireframeMaterial;
+						DrawCollisionColor = (CollisionTraceFlag == ECollisionTraceFlag::CTF_UseComplexAsSimple) ? SimpleCollisionColor : ComplexCollisionColor;
+					}
+					// Create colored proxy
+					FColoredMaterialRenderProxy* CollisionMaterialInstance = new FColoredMaterialRenderProxy(MaterialToUse->GetRenderProxy(), DrawCollisionColor);
+					Collector.RegisterOneFrameMaterialProxy(CollisionMaterialInstance);
+
+					bool bHasPrecomputedVolumetricLightmap;
+					FMatrix PreviousLocalToWorld;
+					int32 SingleCaptureIndex;
+					bool bOutputVelocity;
+					GetScene().GetPrimitiveUniformShaderParameters_RenderThread(GetPrimitiveSceneInfo(), bHasPrecomputedVolumetricLightmap, PreviousLocalToWorld, SingleCaptureIndex, bOutputVelocity);
+					bOutputVelocity |= AlwaysHasVelocity();
+
+					// Draw the mesh with collision materials
+					for (FMeshRenderBufferSet* BufferSet : Buffers)
+					{
+
+						if (BufferSet->TriangleCount == 0)
+						{
+							continue;
+						}
+
+						// lock buffers so that they aren't modified while we are submitting them
+						FScopeLock BuffersLock(&BufferSet->BuffersLock);
+
+						// do we need separate one of these for each MeshRenderBufferSet?
+						FDynamicPrimitiveUniformBuffer& DynamicPrimitiveUniformBuffer = Collector.AllocateOneFrameResource<FDynamicPrimitiveUniformBuffer>();
+						DynamicPrimitiveUniformBuffer.Set(GetLocalToWorld(), PreviousLocalToWorld, GetBounds(), GetLocalBounds(), GetLocalBounds(), true, bHasPrecomputedVolumetricLightmap, bOutputVelocity, GetCustomPrimitiveData());
+
+						if (BufferSet->IndexBuffer.Indices.Num() > 0)
+						{
+							DrawBatch(Collector, *BufferSet, BufferSet->IndexBuffer, CollisionMaterialInstance, bDrawWireframe, SDPG_World, ViewIndex, DynamicPrimitiveUniformBuffer);
+						}
+					}
+				}
+			}
+
+			// Draw simple collision as wireframe if 'show collision', collision is enabled, and we are not using the complex as the simple
+			const bool bDrawSimpleWireframeCollision = (EngineShowFlags.Collision && IsCollisionEnabled() && CollisionTraceFlag != ECollisionTraceFlag::CTF_UseComplexAsSimple);
+
+			if((bDrawSimpleCollision || bDrawSimpleWireframeCollision))
+			{
+				// Note: The static mesh rendering code caches the BodySetup as a class member; if fetching it here is a performance issue, could consider doing the same
+				if (UBodySetup* BodySetup = ParentBaseComponent->GetBodySetup())
+				{
+					// Avoid zero scaling, otherwise GeomTransform below will assert
+					if (FMath::Abs(GetLocalToWorld().Determinant()) > UE_SMALL_NUMBER)
+					{
+						const bool bDrawSolid = !bDrawSimpleWireframeCollision;
+
+						if (AllowDebugViewmodes() && bDrawSolid)
+						{
+							// Make a material for drawing solid collision stuff
+							FColoredMaterialRenderProxy* SolidMaterialInstance = new FColoredMaterialRenderProxy(
+								GEngine->ShadedLevelColorationUnlitMaterial->GetRenderProxy(),
+								GetWireframeColor()
+							);
+
+							Collector.RegisterOneFrameMaterialProxy(SolidMaterialInstance);
+
+							FTransform GeomTransform(GetLocalToWorld());
+							BodySetup->AggGeom.GetAggGeom(GeomTransform, GetWireframeColor().ToFColor(true), SolidMaterialInstance, false, true, AlwaysHasVelocity(), ViewIndex, Collector);
+						}
+						// wireframe
+						else
+						{
+							FTransform GeomTransform(GetLocalToWorld());
+							BodySetup->AggGeom.GetAggGeom(GeomTransform, GetSelectionColor(SimpleCollisionColor, bProxyIsSelected, IsHovered()).ToFColor(true), NULL, bOwnerIsNull, false, AlwaysHasVelocity(), ViewIndex, Collector);
+						}
+
+						// Note: if dynamic mesh component could have nav collision data, we'd also draw that here (see the similar code in StaticMeshRenderer.cpp)
+					}
+				}
+			}
+		}
+	}
+#endif // UE_ENABLE_DEBUG_DRAWING
+
 }
 
 void FBaseDynamicMeshSceneProxy::DrawBatch(FMeshElementCollector& Collector, const FMeshRenderBufferSet& RenderBuffers, const FDynamicMeshIndexBuffer32& IndexBuffer, FMaterialRenderProxy* UseMaterial, bool bWireframe, ESceneDepthPriorityGroup DepthPriority, int ViewIndex, FDynamicPrimitiveUniformBuffer& DynamicPrimitiveUniformBuffer) const
@@ -303,6 +470,19 @@ void FBaseDynamicMeshSceneProxy::DrawBatch(FMeshElementCollector& Collector, con
 	Collector.AddMesh(ViewIndex, Mesh);
 }
 
+void FBaseDynamicMeshSceneProxy::SetCollisionData()
+{
+#if UE_ENABLE_DEBUG_DRAWING
+	bHasCollisionData = true;
+	bOwnerIsNull = ParentBaseComponent->GetOwner() == nullptr;
+	if (UBodySetup* BodySetup = ParentBaseComponent->GetBodySetup())
+	{
+		CollisionTraceFlag = BodySetup->GetCollisionTraceFlag();
+	}
+	CollisionResponse = ParentBaseComponent->GetCollisionResponseToChannels();
+#endif
+}
+
 #if RHI_RAYTRACING
 
 bool FBaseDynamicMeshSceneProxy::IsRayTracingRelevant() const 
diff --git a/Engine/Source/Runtime/GeometryFramework/Public/Components/BaseDynamicMeshSceneProxy.h b/Engine/Source/Runtime/GeometryFramework/Public/Components/BaseDynamicMeshSceneProxy.h
--- a/Engine/Source/Runtime/GeometryFramework/Public/Components/BaseDynamicMeshSceneProxy.h
+++ b/Engine/Source/Runtime/GeometryFramework/Public/Components/BaseDynamicMeshSceneProxy.h
@@ -956,6 +956,16 @@ public:
 		uint32 VisibilityMap, 
 		FMeshElementCollector& Collector) const override;
 
+protected:
+	/**
+	 * Helper called by GetDynamicMeshElements to process collision debug drawing
+	 */
+	GEOMETRYFRAMEWORK_API virtual void GetCollisionDynamicMeshElements(TArray<FMeshRenderBufferSet*>& Buffers,
+		const FEngineShowFlags& EngineShowFlags, bool bDrawCollisionView, bool bDrawSimpleCollision, bool bDrawComplexCollision,
+		bool bProxyIsSelected,
+		const TArray<const FSceneView*>& Views, uint32 VisibilityMap,
+		FMeshElementCollector& Collector) const;
+public:
 
 	/**
 	 * Draw a single-frame FMeshBatch for a FMeshRenderBufferSet
@@ -993,4 +1003,23 @@ public:
 
 #endif // RHI_RAYTRACING
 
+public:
+	// Set the collision data to use for debug drawing, or do nothing if debug drawing is not enabled
+	GEOMETRYFRAMEWORK_API void SetCollisionData();
+
+#if UE_ENABLE_DEBUG_DRAWING
+private:
+	// If debug drawing is enabled, we store collision data here so that collision shapes can be rendered when requested by showflags
+
+	bool bOwnerIsNull = true;
+	/** Whether the collision data has been set up for rendering */
+	bool bHasCollisionData = false;
+	/** Collision trace flags */
+	ECollisionTraceFlag		CollisionTraceFlag;
+	/** Collision Response of this component */
+	FCollisionResponseContainer CollisionResponse;
+#endif
+
+	GEOMETRYFRAMEWORK_API bool IsCollisionView(const FEngineShowFlags& EngineShowFlags, bool& bDrawSimpleCollision, bool& bDrawComplexCollision) const;
+
 };
