Make the dynamic mesh debug collision rendering cache the collision shapes on the dynamic mesh scene proxy, to make sure there is not a race condition if the collision shapes are being updated on the game thread while collision rendering is enabled

diff --git a/Engine/Source/Runtime/GeometryFramework/Private/Components/BaseDynamicMeshSceneProxy.cpp b/Engine/Source/Runtime/GeometryFramework/Private/Components/BaseDynamicMeshSceneProxy.cpp
--- a/Engine/Source/Runtime/GeometryFramework/Private/Components/BaseDynamicMeshSceneProxy.cpp
+++ b/Engine/Source/Runtime/GeometryFramework/Private/Components/BaseDynamicMeshSceneProxy.cpp
@@ -335,7 +335,7 @@ void FBaseDynamicMeshSceneProxy::GetCollisionDynamicMeshElements(TArray<FMeshRen
 		return;
 	}
 
-	// Note: This is closely following StaticMeshRendering.cpp's collision rendering code, from its GetDynamicMeshElements() implementation
+	// Note: This is closely following StaticMeshRender.cpp's collision rendering code, from its GetDynamicMeshElements() implementation
 	FColor SimpleCollisionColor = FColor(157, 149, 223, 255);
 	FColor ComplexCollisionColor = FColor(0, 255, 255, 255);
 
@@ -408,8 +408,7 @@ void FBaseDynamicMeshSceneProxy::GetCollisionDynamicMeshElements(TArray<FMeshRen
 
 			if((bDrawSimpleCollision || bDrawSimpleWireframeCollision))
 			{
-				// Note: The static mesh rendering code caches the BodySetup as a class member; if fetching it here is a performance issue, could consider doing the same
-				if (UBodySetup* BodySetup = ParentBaseComponent->GetBodySetup())
+				if (ParentBaseComponent->GetBodySetup())
 				{
 					// Avoid zero scaling, otherwise GeomTransform below will assert
 					if (FMath::Abs(GetLocalToWorld().Determinant()) > UE_SMALL_NUMBER)
@@ -427,13 +426,13 @@ void FBaseDynamicMeshSceneProxy::GetCollisionDynamicMeshElements(TArray<FMeshRen
 							Collector.RegisterOneFrameMaterialProxy(SolidMaterialInstance);
 
 							FTransform GeomTransform(GetLocalToWorld());
-							BodySetup->AggGeom.GetAggGeom(GeomTransform, GetWireframeColor().ToFColor(true), SolidMaterialInstance, false, true, AlwaysHasVelocity(), ViewIndex, Collector);
+							CachedAggGeom.GetAggGeom(GeomTransform, GetWireframeColor().ToFColor(true), SolidMaterialInstance, false, true, AlwaysHasVelocity(), ViewIndex, Collector);
 						}
 						// wireframe
 						else
 						{
 							FTransform GeomTransform(GetLocalToWorld());
-							BodySetup->AggGeom.GetAggGeom(GeomTransform, GetSelectionColor(SimpleCollisionColor, bProxyIsSelected, IsHovered()).ToFColor(true), NULL, bOwnerIsNull, false, AlwaysHasVelocity(), ViewIndex, Collector);
+							CachedAggGeom.GetAggGeom(GeomTransform, GetSelectionColor(SimpleCollisionColor, bProxyIsSelected, IsHovered()).ToFColor(true), NULL, bOwnerIsNull, false, AlwaysHasVelocity(), ViewIndex, Collector);
 						}
 
 						// Note: if dynamic mesh component could have nav collision data, we'd also draw that here (see the similar code in StaticMeshRenderer.cpp)
@@ -478,6 +477,11 @@ void FBaseDynamicMeshSceneProxy::SetCollisionData()
 	if (UBodySetup* BodySetup = ParentBaseComponent->GetBodySetup())
 	{
 		CollisionTraceFlag = BodySetup->GetCollisionTraceFlag();
+		CachedAggGeom = BodySetup->AggGeom;
+	}
+	else
+	{
+		CachedAggGeom = FKAggregateGeom();
 	}
 	CollisionResponse = ParentBaseComponent->GetCollisionResponseToChannels();
 #endif
diff --git a/Engine/Source/Runtime/GeometryFramework/Public/Components/BaseDynamicMeshSceneProxy.h b/Engine/Source/Runtime/GeometryFramework/Public/Components/BaseDynamicMeshSceneProxy.h
--- a/Engine/Source/Runtime/GeometryFramework/Public/Components/BaseDynamicMeshSceneProxy.h
+++ b/Engine/Source/Runtime/GeometryFramework/Public/Components/BaseDynamicMeshSceneProxy.h
@@ -14,6 +14,8 @@
 #include "Components/BaseDynamicMeshComponent.h"
 #include "RayTracingGeometry.h"
 
+#include "PhysicsEngine/AggregateGeom.h"
+
 using UE::Geometry::FDynamicMesh3;
 using UE::Geometry::FDynamicMeshAttributeSet;
 using UE::Geometry::FDynamicMeshUVOverlay;
@@ -1017,6 +1019,9 @@ private:
 	ECollisionTraceFlag		CollisionTraceFlag;
 	/** Collision Response of this component */
 	FCollisionResponseContainer CollisionResponse;
+	/** Cached AggGeom holding the collision shapes to render */
+	FKAggregateGeom CachedAggGeom;
+
 #endif
 
 	GEOMETRYFRAMEWORK_API bool IsCollisionView(const FEngineShowFlags& EngineShowFlags, bool& bDrawSimpleCollision, bool& bDrawComplexCollision) const;
