Exposed a new property bIsEditable on UDynamicMeshComponent to allow them to block any modifications done through EditMesh/SetMesh. Additionally, when this flag is set to false, the UDynamicMeshToolTarget is blocked from being created for the dynamic mesh component so modeling tools can't edit the mesh either.

diff --git a/Engine/Plugins/Runtime/MeshModelingToolset/Source/ModelingComponentsEditorOnly/Private/ToolTargets/DynamicMeshComponentToolTarget.cpp b/Engine/Plugins/Runtime/MeshModelingToolset/Source/ModelingComponentsEditorOnly/Private/ToolTargets/DynamicMeshComponentToolTarget.cpp
--- a/Engine/Plugins/Runtime/MeshModelingToolset/Source/ModelingComponentsEditorOnly/Private/ToolTargets/DynamicMeshComponentToolTarget.cpp
+++ b/Engine/Plugins/Runtime/MeshModelingToolset/Source/ModelingComponentsEditorOnly/Private/ToolTargets/DynamicMeshComponentToolTarget.cpp
@@ -230,6 +230,7 @@ bool UDynamicMeshComponentToolTargetFactory::CanBuildTarget(UObject* SourceObjec
 		&& !Component->IsUnreachable() 
 		&& Component->IsValidLowLevel() 
 		&& Component->GetDynamicMesh()
+		&& Component->IsEditable()
 		&& Requirements.AreSatisfiedBy(UDynamicMeshComponentToolTarget::StaticClass());
 }
 
diff --git a/Engine/Source/Runtime/GeometryFramework/Private/Components/DynamicMeshComponent.cpp b/Engine/Source/Runtime/GeometryFramework/Private/Components/DynamicMeshComponent.cpp
--- a/Engine/Source/Runtime/GeometryFramework/Private/Components/DynamicMeshComponent.cpp
+++ b/Engine/Source/Runtime/GeometryFramework/Private/Components/DynamicMeshComponent.cpp
@@ -155,7 +155,10 @@ void UDynamicMeshComponent::PostEditChangeProperty(FPropertyChangedEvent& Proper
 
 void UDynamicMeshComponent::SetMesh(UE::Geometry::FDynamicMesh3&& MoveMesh)
 {
-	MeshObject->SetMesh(MoveTemp(MoveMesh));
+	if (ensureMsgf(IsEditable(), TEXT("Attempted to modify the internal mesh of a UDynamicMeshComponent that is not editable")))
+	{
+		MeshObject->SetMesh(MoveTemp(MoveMesh));
+	}
 }
 
 
@@ -169,10 +172,13 @@ void UDynamicMeshComponent::ProcessMesh(
 void UDynamicMeshComponent::EditMesh(TFunctionRef<void(UE::Geometry::FDynamicMesh3&)> EditFunc,
 										   EDynamicMeshComponentRenderUpdateMode UpdateMode )
 {
-	MeshObject->EditMesh(EditFunc);
-	if (UpdateMode != EDynamicMeshComponentRenderUpdateMode::NoUpdate)
+	if (ensureMsgf(IsEditable(), TEXT("Attempted to modify the internal mesh of a UDynamicMeshComponent that is not editable")))
 	{
-		NotifyMeshUpdated();
+		MeshObject->EditMesh(EditFunc);
+		if (UpdateMode != EDynamicMeshComponentRenderUpdateMode::NoUpdate)
+		{
+			NotifyMeshUpdated();
+		}
 	}
 }
 
@@ -223,17 +229,20 @@ const FDynamicMesh3* UDynamicMeshComponent::GetRenderMesh() const
 
 void UDynamicMeshComponent::ApplyTransform(const FTransform3d& Transform, bool bInvert)
 {
-	MeshObject->EditMesh([&](FDynamicMesh3& EditMesh)
+	if (ensureMsgf(IsEditable(), TEXT("Attempted to modify the internal mesh of a UDynamicMeshComponent that is not editable")))
 	{
-		if (bInvert)
-		{
-			MeshTransforms::ApplyTransformInverse(EditMesh, Transform, true);
-		}
-		else
+		MeshObject->EditMesh([&](FDynamicMesh3& EditMesh)
 		{
-			MeshTransforms::ApplyTransform(EditMesh, Transform, true);
-		}
-	}, EDynamicMeshChangeType::DeformationEdit);
+			if (bInvert)
+			{
+				MeshTransforms::ApplyTransformInverse(EditMesh, Transform, true);
+			}
+			else
+			{
+				MeshTransforms::ApplyTransform(EditMesh, Transform, true);
+			}
+		}, EDynamicMeshChangeType::DeformationEdit);
+	}
 }
 
 
diff --git a/Engine/Source/Runtime/GeometryFramework/Public/Components/DynamicMeshComponent.h b/Engine/Source/Runtime/GeometryFramework/Public/Components/DynamicMeshComponent.h
--- a/Engine/Source/Runtime/GeometryFramework/Public/Components/DynamicMeshComponent.h
+++ b/Engine/Source/Runtime/GeometryFramework/Public/Components/DynamicMeshComponent.h
@@ -119,6 +119,12 @@ public:
 	GEOMETRYFRAMEWORK_API virtual void ApplyTransform(const FTransform3d& Transform, bool bInvert) override;
 
 
+	/**
+	 * If set to false, this blocks external code from being able to modify the internal mesh. It prevents modifications through methods on this UDynamicMeshComponent such as EditMesh/SetMesh,
+	 * However, this cannot prevent code from modifying the underlying UDynamicMesh through direct references (ex. via GetDynamicMesh()->SetMesh).
+	 */
+	GEOMETRYFRAMEWORK_API bool IsEditable() const { return bIsEditable; }
+	GEOMETRYFRAMEWORK_API void SetIsEditable(bool bInIsEditable) { bIsEditable = bInIsEditable; }
 
 protected:
 	/**
@@ -296,6 +302,11 @@ protected:
 	/** Called whenever internal MeshObject is modified, fires OnMeshChanged and OnMeshVerticesChanged above */
 	GEOMETRYFRAMEWORK_API void OnMeshObjectChanged(UDynamicMesh* ChangedMeshObject, FDynamicMeshChangeInfo ChangeInfo);
 
+	/**
+	 * If set to false, this blocks external code from being able to modify the internal mesh. It prevents modifications through methods on this UDynamicMeshComponent such as EditMesh/SetMesh,
+	 * However, this cannot prevent code from modifying the underlying UDynamicMesh through direct references (ex. via GetDynamicMesh()->SetMesh).
+	 */
+	bool bIsEditable = true;
 
 
 
