Geometry Framework.  allow SetMaterial() to trigger a immediate visual update.

diff --git a/Engine/Source/Runtime/GeometryFramework/Private/Components/BaseDynamicMeshComponent.cpp b/Engine/Source/Runtime/GeometryFramework/Private/Components/BaseDynamicMeshComponent.cpp
--- a/Engine/Source/Runtime/GeometryFramework/Private/Components/BaseDynamicMeshComponent.cpp
+++ b/Engine/Source/Runtime/GeometryFramework/Private/Components/BaseDynamicMeshComponent.cpp
@@ -216,6 +216,19 @@ void UBaseDynamicMeshComponent::SetMaterial(int32 ElementIndex, UMaterialInterfa
 		BaseMaterials.SetNum(ElementIndex + 1, EAllowShrinking::No);
 	}
 	BaseMaterials[ElementIndex] = Material;
+
+	// @todo allow for precache of pipeline state objects for rendering
+	// PrecachePSOs(); // indirectly calls GetUsedMaterials, requires CollectPSOPrecacheData to be implemented, see UStaticMeshComponent for example
+
+	
+	MarkRenderStateDirty();
+
+	// update the body instance in case this material has an associated physics material 
+	FBodyInstance* BodyInst = GetBodyInstance();
+	if (BodyInst && BodyInst->IsValidBodyInstance())
+	{
+		BodyInst->UpdatePhysicalMaterials();
+	}
 }
 
 
