Geometry Tools: TransfromUVs,  fixing undo within the tool.
GeometryFramework: Adding primary UVs to FVertexChange data.

diff --git a/Engine/Plugins/Experimental/MeshModelingToolsetExp/Source/MeshModelingToolsExp/Private/EditUVIslandsTool.cpp b/Engine/Plugins/Experimental/MeshModelingToolsetExp/Source/MeshModelingToolsExp/Private/EditUVIslandsTool.cpp
--- a/Engine/Plugins/Experimental/MeshModelingToolsetExp/Source/MeshModelingToolsExp/Private/EditUVIslandsTool.cpp
+++ b/Engine/Plugins/Experimental/MeshModelingToolsetExp/Source/MeshModelingToolsExp/Private/EditUVIslandsTool.cpp
@@ -557,8 +557,12 @@ void UEditUVIslandsTool::UpdateChangeFromROI(bool bFinal)
 	}
 
 	FDynamicMesh3* Mesh = DynamicMeshComponent->GetMesh();
-	//const TSet<int>& ModifiedVertices = LinearDeformer.GetModifiedVertices();
-	//ActiveVertexChange->SavePositions(Mesh, ModifiedVertices, !bFinal);
+	TArray<int32> ModifiedUVs;
+	for (FEditIsland& Island : ActiveIslands)
+	{
+		ModifiedUVs.Append(Island.UVs);
+	}
+	ActiveVertexChange->SaveOverlayUVs(Mesh, ModifiedUVs, !bFinal);
 }
 
 
@@ -566,7 +570,7 @@ void UEditUVIslandsTool::BeginChange()
 {
 	if (ActiveVertexChange == nullptr)
 	{
-		ActiveVertexChange = new FMeshVertexChangeBuilder();
+		ActiveVertexChange = new FMeshVertexChangeBuilder(EMeshVertexChangeComponents::OverylayUVs);
 		UpdateChangeFromROI(false);
 	}
 }
diff --git a/Engine/Source/Runtime/GeometryFramework/Private/Changes/MeshVertexChange.cpp b/Engine/Source/Runtime/GeometryFramework/Private/Changes/MeshVertexChange.cpp
--- a/Engine/Source/Runtime/GeometryFramework/Private/Changes/MeshVertexChange.cpp
+++ b/Engine/Source/Runtime/GeometryFramework/Private/Changes/MeshVertexChange.cpp
@@ -46,6 +46,9 @@ FMeshVertexChangeBuilder::FMeshVertexChangeBuilder(EMeshVertexChangeComponents C
 	
 	bSaveOverlayNormals = ((Components & EMeshVertexChangeComponents::OverlayNormals) != EMeshVertexChangeComponents::None);
 	Change->bHaveOverlayNormals = bSaveOverlayNormals;
+
+	bSaveOverlayUVs = ((Components & EMeshVertexChangeComponents::OverylayUVs) != EMeshVertexChangeComponents::None);
+	Change->bHaveOverlayUVs = bSaveOverlayUVs;
 }
 
 void FMeshVertexChangeBuilder::UpdateVertex(int VertexID, const FVector3d& OldPosition, const FVector3d& NewPosition)
@@ -249,3 +252,83 @@ void FMeshVertexChangeBuilder::SaveOverlayNormals(const FDynamicMesh3* Mesh, con
 		}
 	}
 }
+
+void FMeshVertexChangeBuilder::UpdateOverlayUV(int ElementID, const FVector2f& OldUV, const FVector2f& NewUV)
+{
+	const int* FoundIndex = SavedUVElements.Find(ElementID);
+	if (FoundIndex == nullptr)
+	{
+		int NewIndex = Change->UVs.Num();
+		SavedUVElements.Add(ElementID, NewIndex);
+		Change->UVs.Add(ElementID);
+		Change->OldUVs.Add(OldUV);
+		Change->NewUVs.Add(NewUV);
+	}
+	else
+	{
+		Change->NewUVs[*FoundIndex] = NewUV;
+	}
+}
+
+void FMeshVertexChangeBuilder::UpdateOverlayUVFinal(int ElementID, const FVector2f& NewUV)
+{
+	check(SavedUVElements.Contains(ElementID));
+	const int* Index = SavedUVElements.Find(ElementID);
+	if (Index != nullptr)
+	{
+		Change->NewUVs[*Index] = NewUV;
+	}
+}
+
+void FMeshVertexChangeBuilder::SaveOverlayUVs(const FDynamicMesh3* Mesh, const TArray<int>& ElementIDs, bool bInitial)
+{
+	if (Mesh->HasAttributes() == false || Mesh->Attributes()->PrimaryUV() == nullptr)
+	{
+		return;
+	}
+	const FDynamicMeshUVOverlay* Overlay = Mesh->Attributes()->PrimaryUV();
+
+	int Num = ElementIDs.Num();
+	if (bInitial)
+	{
+		for (int k = 0; k < Num; ++k)
+		{
+			FVector2f UV = Overlay->GetElement(ElementIDs[k]);
+			UpdateOverlayUV(ElementIDs[k], UV, UV);
+		}
+	}
+	else
+	{
+		for (int k = 0; k < Num; ++k)
+		{
+			FVector2f UV = Overlay->GetElement(ElementIDs[k]);
+			UpdateOverlayUVFinal(ElementIDs[k], UV);
+		}
+	}
+}
+
+void FMeshVertexChangeBuilder::SaveOverlayUVs(const FDynamicMesh3* Mesh, const TSet<int>& ElementIDs, bool bInitial)
+{
+	if (Mesh->HasAttributes() == false || Mesh->Attributes()->PrimaryUV() == nullptr)
+	{
+		return;
+	}
+	const FDynamicMeshUVOverlay* Overlay = Mesh->Attributes()->PrimaryUV();
+
+	if (bInitial)
+	{
+		for (int ElementID : ElementIDs)
+		{
+			FVector2f UV = Overlay->GetElement(ElementID);
+			UpdateOverlayUV(ElementID, UV, UV);
+		}
+	}
+	else
+	{
+		for (int ElementID : ElementIDs)
+		{
+			FVector2f UV = Overlay->GetElement(ElementID);
+			UpdateOverlayUVFinal(ElementID, UV);
+		}
+	}
+}
\ No newline at end of file
diff --git a/Engine/Source/Runtime/GeometryFramework/Private/UDynamicMesh.cpp b/Engine/Source/Runtime/GeometryFramework/Private/UDynamicMesh.cpp
--- a/Engine/Source/Runtime/GeometryFramework/Private/UDynamicMesh.cpp
+++ b/Engine/Source/Runtime/GeometryFramework/Private/UDynamicMesh.cpp
@@ -240,6 +240,21 @@ void UDynamicMesh::ApplyChange(const FMeshVertexChange* Change, bool bRevert)
 				}
 			}
 		}
+
+		if (Change->bHaveOverlayUVs && EditMesh.HasAttributes() && EditMesh.Attributes()->PrimaryUV())
+		{
+			FDynamicMeshUVOverlay* Overlay = EditMesh.Attributes()->PrimaryUV();
+			int32 NumUVs = Change->UVs.Num();
+			const TArray<FVector2f>& UseUVs = (bRevert) ? Change->OldUVs : Change->NewUVs;
+			for (int32 k = 0; k < NumUVs; ++k)
+			{
+				int32 elemid = Change->UVs[k];
+				if (Overlay->IsElement(elemid))
+				{
+					Overlay->SetElement(elemid, UseUVs[k]);
+				}
+			}
+		}
 	}, ChangeInfo);
 }
 
diff --git a/Engine/Source/Runtime/GeometryFramework/Public/Changes/MeshVertexChange.h b/Engine/Source/Runtime/GeometryFramework/Public/Changes/MeshVertexChange.h
--- a/Engine/Source/Runtime/GeometryFramework/Public/Changes/MeshVertexChange.h
+++ b/Engine/Source/Runtime/GeometryFramework/Public/Changes/MeshVertexChange.h
@@ -11,7 +11,7 @@
 PREDECLARE_USE_GEOMETRY_CLASS(FDynamicMesh3);
 
 /**
- * FMeshVertexChange represents an reversible change to a set of vertex positions.
+ * FMeshVertexChange represents an reversible change to a set of vertex positions, normals, colors and UVs
  * Currently only a UDynamicMeshComponent target is supported.
  * 
  * @todo support optionally storing old/new normals and tangents
@@ -33,6 +33,11 @@ public:
 	TArray<FVector3f> OldNormals;
 	TArray<FVector3f> NewNormals;
 
+	bool bHaveOverlayUVs = false;
+	TArray<int32> UVs;
+	TArray<FVector2f> OldUVs;
+	TArray<FVector2f> NewUVs;
+
 	/** Makes the change to the object */
 	GEOMETRYFRAMEWORK_API virtual void Apply(UObject* Object) override;
 
@@ -49,6 +54,7 @@ enum class EMeshVertexChangeComponents : uint8
 	None = 0,
 	VertexPositions = 1,
 	VertexColors = 2,
+	OverylayUVs = 4,
 	OverlayNormals = 16
 };
 ENUM_CLASS_FLAGS(EMeshVertexChangeComponents);
@@ -68,6 +74,9 @@ public:
 	bool bSaveOverlayNormals = false;
 	TMap<int32, int32> SavedNormalElements;
 
+	bool bSaveOverlayUVs = false;
+	TMap<int32, int32> SavedUVElements;
+
 	/** If set, this function is called whenever a newly-seen VertexID is saved, parameters are (VertexID, Index) into saved-vertices array */
 	TUniqueFunction<void(int32, int32)> OnNewVertexSaved = nullptr;
 
@@ -83,6 +92,8 @@ public:
 	GEOMETRYFRAMEWORK_API void SaveOverlayNormals(const FDynamicMesh3* Mesh, const TArray<int32>& ElementIDs, bool bInitial);
 	GEOMETRYFRAMEWORK_API void SaveOverlayNormals(const FDynamicMesh3* Mesh, const TSet<int32>& ElementIDs, bool bInitial);
 
+	GEOMETRYFRAMEWORK_API void SaveOverlayUVs(const FDynamicMesh3* Mesh, const TArray<int32>& ElementIDs, bool bInitial);
+	GEOMETRYFRAMEWORK_API void SaveOverlayUVs(const FDynamicMesh3* Mesh, const TSet<int32>& ElementIDs, bool bInitial);
 
 public:
 	// currently only used in vertex sculpt tool. cannot be used if bSaveColors = true
@@ -96,6 +107,9 @@ protected:
 
 	GEOMETRYFRAMEWORK_API void UpdateOverlayNormal(int32 ElementID, const FVector3f& OldNormal, const FVector3f& NewNormal);
 	GEOMETRYFRAMEWORK_API void UpdateOverlayNormalFinal(int32 ElementID, const FVector3f& NewNormal);
+
+	GEOMETRYFRAMEWORK_API void UpdateOverlayUV(int32 ElementID, const FVector2f& OldUV, const FVector2f& NewUV);
+	GEOMETRYFRAMEWORK_API void UpdateOverlayUVFinal(int32 ElementID, const FVector2f& NewUV);
 };
 
 
