add physics materials to dynamic mesh component's physics mesh

diff --git a/Engine/Source/Runtime/GeometryFramework/Private/Components/DynamicMeshComponent.cpp b/Engine/Source/Runtime/GeometryFramework/Private/Components/DynamicMeshComponent.cpp
--- a/Engine/Source/Runtime/GeometryFramework/Private/Components/DynamicMeshComponent.cpp
+++ b/Engine/Source/Runtime/GeometryFramework/Private/Components/DynamicMeshComponent.cpp
@@ -1248,7 +1248,6 @@ void UDynamicMeshComponent::OnChildDetached(USceneComponent* ChildComponent)
 
 bool UDynamicMeshComponent::GetPhysicsTriMeshData(struct FTriMeshCollisionData* CollisionData, bool InUseAllTriData)
 {
-	// todo: support physical materials
 	// todo: support UPhysicsSettings::Get()->bSupportUVFromHitResults
 
 	// this is something we currently assume, if you hit this ensure, we made a mistake
@@ -1256,6 +1255,8 @@ bool UDynamicMeshComponent::GetPhysicsTriMeshData(struct FTriMeshCollisionData*
 
 	ProcessMesh([&](const FDynamicMesh3& Mesh)
 	{
+		const FDynamicMeshMaterialAttribute* MaterialAttrib = Mesh.HasAttributes() && Mesh.Attributes()->HasMaterialID() ? Mesh.Attributes()->GetMaterialID() : nullptr;
+
 		TArray<int32> VertexMap;
 		bool bIsSparseV = !Mesh.IsCompactV();
 		if (bIsSparseV)
@@ -1307,7 +1308,8 @@ bool UDynamicMeshComponent::GetPhysicsTriMeshData(struct FTriMeshCollisionData*
 
 			CollisionData->Indices.Add(Triangle);
 
-			CollisionData->MaterialIndices.Add(0);		// not supporting physical materials yet
+			int32 MaterialID = MaterialAttrib ? MaterialAttrib->GetValue(tid) : 0;
+			CollisionData->MaterialIndices.Add(MaterialID);
 		}
 
 		CollisionData->bFlipNormals = true;
