add UVs to dynamic mesh component's physics mesh

diff --git a/Engine/Source/Runtime/GeometryFramework/Private/Components/DynamicMeshComponent.cpp b/Engine/Source/Runtime/GeometryFramework/Private/Components/DynamicMeshComponent.cpp
--- a/Engine/Source/Runtime/GeometryFramework/Private/Components/DynamicMeshComponent.cpp
+++ b/Engine/Source/Runtime/GeometryFramework/Private/Components/DynamicMeshComponent.cpp
@@ -8,6 +8,7 @@
 #include "Materials/Material.h"
 #include "Async/Async.h"
 #include "Engine/CollisionProfile.h"
+#include "PhysicsEngine/PhysicsSettings.h"
 
 #include "DynamicMesh/DynamicMeshAttributeSet.h"
 #include "DynamicMesh/MeshNormals.h"
@@ -1243,39 +1244,125 @@ void UDynamicMeshComponent::OnChildDetached(USceneComponent* ChildComponent)
 }
 
 
-
-
+bool UDynamicMeshComponent::GetTriMeshSizeEstimates(struct FTriMeshCollisionDataEstimates& OutTriMeshEstimates, bool bInUseAllTriData) const
+{
+	ProcessMesh([&](const FDynamicMesh3& Mesh)
+		{
+			bool bCopyUVs = UPhysicsSettings::Get()->bSupportUVFromHitResults && Mesh.HasAttributes() && Mesh.Attributes()->NumUVLayers() > 0;
+			if (bCopyUVs)
+			{
+				// conservative estimate
+				OutTriMeshEstimates.VerticeCount = Mesh.TriangleCount() * 3;
+			}
+			else
+			{
+				OutTriMeshEstimates.VerticeCount = Mesh.VertexCount();
+			}
+		}
+	);
+	return true;
+}
 
 bool UDynamicMeshComponent::GetPhysicsTriMeshData(struct FTriMeshCollisionData* CollisionData, bool InUseAllTriData)
 {
-	// todo: support UPhysicsSettings::Get()->bSupportUVFromHitResults
-
 	// this is something we currently assume, if you hit this ensure, we made a mistake
 	ensure(bEnableComplexCollision);
 
 	ProcessMesh([&](const FDynamicMesh3& Mesh)
 	{
-		const FDynamicMeshMaterialAttribute* MaterialAttrib = Mesh.HasAttributes() && Mesh.Attributes()->HasMaterialID() ? Mesh.Attributes()->GetMaterialID() : nullptr;
-
-		TArray<int32> VertexMap;
-		bool bIsSparseV = !Mesh.IsCompactV();
-		if (bIsSparseV)
+		// See if we should copy UVs
+		const bool bCopyUVs = UPhysicsSettings::Get()->bSupportUVFromHitResults && Mesh.HasAttributes() && Mesh.Attributes()->NumUVLayers() > 0;
+		if (bCopyUVs)
 		{
-			VertexMap.SetNum(Mesh.MaxVertexID());
+			CollisionData->UVs.SetNum(Mesh.Attributes()->NumUVLayers());
 		}
+		const FDynamicMeshMaterialAttribute* MaterialAttrib = Mesh.HasAttributes() && Mesh.Attributes()->HasMaterialID() ? Mesh.Attributes()->GetMaterialID() : nullptr;
+
+		TArray<int32> VertexMap; 
+		const bool bIsSparseV = !Mesh.IsCompactV();
 
 		// copy vertices
-		CollisionData->Vertices.Reserve(Mesh.VertexCount());
-		for (int32 vid : Mesh.VertexIndicesItr())
+		if (!bCopyUVs)
 		{
-			int32 Index = CollisionData->Vertices.Add((FVector3f)Mesh.GetVertex(vid));
 			if (bIsSparseV)
 			{
-				VertexMap[vid] = Index;
+				VertexMap.SetNum(Mesh.MaxVertexID());
 			}
-			else
+			CollisionData->Vertices.Reserve(Mesh.VertexCount());
+			for (int32 vid : Mesh.VertexIndicesItr())
+			{
+				int32 Index = CollisionData->Vertices.Add((FVector3f)Mesh.GetVertex(vid));
+				if (bIsSparseV)
+				{
+					VertexMap[vid] = Index;
+				}
+				else
+				{
+					check(vid == Index);
+				}
+			}
+		}
+		else
+		{
+			// map vertices per wedge
+			VertexMap.SetNumZeroed(Mesh.TriangleCount() * 3);
+			// temp array to store the UVs on a vertex (per triangle)
+			TArray<FVector2D> VertUVs;
+			const FDynamicMeshAttributeSet* Attribs = Mesh.Attributes();
+			const int32 NumUVLayers = Attribs->NumUVLayers();
+			for (int32 VID : Mesh.VertexIndicesItr())
 			{
-				check(vid == Index);
+				FVector3f Pos = (FVector3f)Mesh.GetVertex(VID);
+				int32 VertStart = CollisionData->Vertices.Num();
+				Mesh.EnumerateVertexTriangles(VID, [&](int32 TID)
+				{
+					FIndex3i Tri = Mesh.GetTriangle(TID);
+					int32 VSubIdx = Tri.IndexOf(VID);
+					// Get the UVs on this wedge
+					VertUVs.Reset(8);
+					for (int32 UVIdx = 0; UVIdx < NumUVLayers; ++UVIdx)
+					{
+						const FDynamicMeshUVOverlay* Overlay = Attribs->GetUVLayer(UVIdx);
+						FIndex3i UVTri = Overlay->GetTriangle(TID);
+						int32 ElID = UVTri[VSubIdx];
+						FVector2D UV(0, 0);
+						if (ElID >= 0)
+						{
+							UV = (FVector2D)Overlay->GetElement(ElID);
+						}
+						VertUVs.Add(UV);
+					}
+					// Check if we've already added these UVs via an earlier wedge
+					int32 OutputVIdx = INDEX_NONE;
+					for (int32 VIdx = VertStart; VIdx < CollisionData->Vertices.Num(); ++VIdx)
+					{
+						bool bFound = true;
+						for (int32 UVIdx = 0; UVIdx < NumUVLayers; ++UVIdx)
+						{
+							if (CollisionData->UVs[UVIdx][VIdx] != VertUVs[UVIdx])
+							{
+								bFound = false;
+								break;
+							}
+						}
+						if (bFound)
+						{
+							OutputVIdx = VIdx;
+							break;
+						}
+					}
+					// If not, add the vertex w/ the UVs
+					if (OutputVIdx == INDEX_NONE)
+					{
+						OutputVIdx = CollisionData->Vertices.Add(Pos);
+						for (int32 UVIdx = 0; UVIdx < NumUVLayers; ++UVIdx)
+						{
+							CollisionData->UVs[UVIdx].Add(VertUVs[UVIdx]);
+						}
+					}
+					// Map the wedge to the output vertex
+					VertexMap[TID * 3 + VSubIdx] = OutputVIdx;
+				});
 			}
 		}
 
@@ -1286,9 +1373,25 @@ bool UDynamicMeshComponent::GetPhysicsTriMeshData(struct FTriMeshCollisionData*
 		{
 			FIndex3i Tri = Mesh.GetTriangle(tid);
 			FTriIndices Triangle;
-			Triangle.v0 = (bIsSparseV) ? VertexMap[Tri.A] : Tri.A;
-			Triangle.v1 = (bIsSparseV) ? VertexMap[Tri.B] : Tri.B;
-			Triangle.v2 = (bIsSparseV) ? VertexMap[Tri.C] : Tri.C;
+			if (bCopyUVs)
+			{
+				// UVs need a wedge-based map
+				Triangle.v0 = VertexMap[tid * 3 + 0];
+				Triangle.v1 = VertexMap[tid * 3 + 1];
+				Triangle.v2 = VertexMap[tid * 3 + 2];
+			}
+			else if (bIsSparseV)
+			{
+				Triangle.v0 = VertexMap[Tri.A];
+				Triangle.v1 = VertexMap[Tri.B];
+				Triangle.v2 = VertexMap[Tri.C];
+			}
+			else
+			{ 
+				Triangle.v0 = Tri.A;
+				Triangle.v1 = Tri.B;
+				Triangle.v2 = Tri.C;
+			}
 
 			// Filter out triangles which will cause physics system to emit degenerate-geometry warnings.
 			// These checks reproduce tests in Chaos::CleanTrimesh
diff --git a/Engine/Source/Runtime/GeometryFramework/Public/Components/DynamicMeshComponent.h b/Engine/Source/Runtime/GeometryFramework/Public/Components/DynamicMeshComponent.h
--- a/Engine/Source/Runtime/GeometryFramework/Public/Components/DynamicMeshComponent.h
+++ b/Engine/Source/Runtime/GeometryFramework/Public/Components/DynamicMeshComponent.h
@@ -563,6 +563,7 @@ public:
 
 
 	GEOMETRYFRAMEWORK_API virtual bool GetPhysicsTriMeshData(struct FTriMeshCollisionData* CollisionData, bool InUseAllTriData) override;
+	GEOMETRYFRAMEWORK_API virtual bool GetTriMeshSizeEstimates(struct FTriMeshCollisionDataEstimates& OutTriMeshEstimates, bool bInUseAllTriData) const override;
 	GEOMETRYFRAMEWORK_API virtual bool ContainsPhysicsTriMeshData(bool InUseAllTriData) const override;
 	GEOMETRYFRAMEWORK_API virtual bool WantsNegXTriMesh() override;
 
