Call FRayTracingGeometry::IsValid() instead of just checking that RayTracingGeometryRHI is not null

diff --git a/Engine/Plugins/Experimental/ImagePlate/Source/ImagePlate/Private/ImagePlateComponent.cpp b/Engine/Plugins/Experimental/ImagePlate/Source/ImagePlate/Private/ImagePlateComponent.cpp
--- a/Engine/Plugins/Experimental/ImagePlate/Source/ImagePlate/Private/ImagePlateComponent.cpp
+++ b/Engine/Plugins/Experimental/ImagePlate/Source/ImagePlate/Private/ImagePlateComponent.cpp
@@ -263,7 +263,7 @@ namespace
 
 			FMaterialRenderProxy* MaterialProxy = Material->GetRenderProxy();
 
-			if (RayTracingGeometry.RayTracingGeometryRHI.IsValid())
+			if (RayTracingGeometry.IsValid())
 			{
 				check(RayTracingGeometry.Initializer.IndexBuffer.IsValid());
 
diff --git a/Engine/Plugins/Runtime/HairStrands/Source/HairStrandsCore/Private/GroomComponent.cpp b/Engine/Plugins/Runtime/HairStrands/Source/HairStrandsCore/Private/GroomComponent.cpp
--- a/Engine/Plugins/Runtime/HairStrands/Source/HairStrandsCore/Private/GroomComponent.cpp
+++ b/Engine/Plugins/Runtime/HairStrands/Source/HairStrandsCore/Private/GroomComponent.cpp
@@ -697,7 +697,7 @@ public:
 				continue;
 			}
 
-			if (RTGeometry && RTGeometry->RayTracingGeometry.RayTracingGeometryRHI.IsValid())
+			if (RTGeometry && RTGeometry->RayTracingGeometry.IsValid())
 			{
 				for (const FRayTracingGeometrySegment& Segment : RTGeometry->RayTracingGeometry.Initializer.Segments)
 				{
diff --git a/Engine/Plugins/Runtime/ProceduralMeshComponent/Source/ProceduralMeshComponent/Private/ProceduralMeshComponent.cpp b/Engine/Plugins/Runtime/ProceduralMeshComponent/Source/ProceduralMeshComponent/Private/ProceduralMeshComponent.cpp
--- a/Engine/Plugins/Runtime/ProceduralMeshComponent/Source/ProceduralMeshComponent/Private/ProceduralMeshComponent.cpp
+++ b/Engine/Plugins/Runtime/ProceduralMeshComponent/Source/ProceduralMeshComponent/Private/ProceduralMeshComponent.cpp
@@ -464,7 +464,7 @@ public:
 			{
 				FMaterialRenderProxy* MaterialProxy = Section->Material->GetRenderProxy();
 				
-				if (Section->RayTracingGeometry.RayTracingGeometryRHI.IsValid())
+				if (Section->RayTracingGeometry.IsValid())
 				{
 					check(Section->RayTracingGeometry.Initializer.IndexBuffer.IsValid());
 
diff --git a/Engine/Source/Runtime/Engine/Private/SkeletalMesh.cpp b/Engine/Source/Runtime/Engine/Private/SkeletalMesh.cpp
--- a/Engine/Source/Runtime/Engine/Private/SkeletalMesh.cpp
+++ b/Engine/Source/Runtime/Engine/Private/SkeletalMesh.cpp
@@ -6909,7 +6909,7 @@ void FSkeletalMeshSceneProxy::GetDynamicRayTracingInstances(FRayTracingMaterialG
 	// GetRayTracingGeometry()->IsInitialized() is checked as a workaround for UE-92634. FSkeletalMeshSceneProxy's resources may have already been released, but proxy has not removed yet)
 	if (MeshObject->GetRayTracingGeometry() && MeshObject->GetRayTracingGeometry()->IsInitialized())
 	{
-		if(MeshObject->GetRayTracingGeometry()->RayTracingGeometryRHI.IsValid())
+		if(MeshObject->GetRayTracingGeometry()->IsValid())
 		{
 			check(MeshObject->GetRayTracingGeometry()->Initializer.IndexBuffer.IsValid());
 			
diff --git a/Engine/Source/Runtime/Engine/Private/Streaming/StaticMeshUpdate.cpp b/Engine/Source/Runtime/Engine/Private/Streaming/StaticMeshUpdate.cpp
--- a/Engine/Source/Runtime/Engine/Private/Streaming/StaticMeshUpdate.cpp
+++ b/Engine/Source/Runtime/Engine/Private/Streaming/StaticMeshUpdate.cpp
@@ -303,7 +303,7 @@ void FStaticMeshStreamIn::DoFinishUpdate(const FContext& Context)
 #endif
 
 				// Under very rare circumstances that we switch ray tracing on/off right in the middle of streaming RayTracingGeometryRHI might not be valid.
-				if (IsRayTracingEnabled() && ensure(LODResource.RayTracingGeometry.RayTracingGeometryRHI.IsValid()))
+				if (IsRayTracingEnabled() && ensure(LODResource.RayTracingGeometry.IsValid()))
 				{
 					LODResource.RayTracingGeometry.RequestBuildIfNeeded(ERTAccelerationStructureBuildPriority::Normal);
 				}
diff --git a/Engine/Source/Runtime/Experimental/GeometryCollectionEngine/Private/GeometryCollection/GeometryCollectionSceneProxy.cpp b/Engine/Source/Runtime/Experimental/GeometryCollectionEngine/Private/GeometryCollection/GeometryCollectionSceneProxy.cpp
--- a/Engine/Source/Runtime/Experimental/GeometryCollectionEngine/Private/GeometryCollection/GeometryCollectionSceneProxy.cpp
+++ b/Engine/Source/Runtime/Experimental/GeometryCollectionEngine/Private/GeometryCollection/GeometryCollectionSceneProxy.cpp
@@ -817,7 +817,7 @@ void FGeometryCollectionSceneProxy::GetDynamicRayTracingInstances(FRayTracingMat
 
 	UpdatingRayTracingGeometry_RenderingThread(SectionArray);
 
-	if (SectionArray.Num() && RayTracingGeometry.RayTracingGeometryRHI.IsValid())
+	if (SectionArray.Num() && RayTracingGeometry.IsValid())
 	{
 		FRayTracingInstance RayTracingInstance;
 		RayTracingInstance.Geometry = &RayTracingGeometry;
diff --git a/Engine/Source/Runtime/GeometryFramework/Private/Components/BaseDynamicMeshSceneProxy.cpp b/Engine/Source/Runtime/GeometryFramework/Private/Components/BaseDynamicMeshSceneProxy.cpp
--- a/Engine/Source/Runtime/GeometryFramework/Private/Components/BaseDynamicMeshSceneProxy.cpp
+++ b/Engine/Source/Runtime/GeometryFramework/Private/Components/BaseDynamicMeshSceneProxy.cpp
@@ -565,7 +565,7 @@ void FBaseDynamicMeshSceneProxy::GetDynamicRayTracingInstances(FRayTracingMateri
 
 		// draw primary index buffer
 		if (BufferSet->IndexBuffer.Indices.Num() > 0
-			&& BufferSet->PrimaryRayTracingGeometry.RayTracingGeometryRHI.IsValid())
+			&& BufferSet->PrimaryRayTracingGeometry.IsValid())
 		{
 			ensure(BufferSet->PrimaryRayTracingGeometry.Initializer.IndexBuffer.IsValid());
 			DrawRayTracingBatch(Context, *BufferSet, BufferSet->IndexBuffer, BufferSet->PrimaryRayTracingGeometry, MaterialProxy, DepthPriority, DynamicPrimitiveUniformBuffer, OutRayTracingInstances);
@@ -576,7 +576,7 @@ void FBaseDynamicMeshSceneProxy::GetDynamicRayTracingInstances(FRayTracingMateri
 		if (bDrawSecondaryBuffers
 			&& BufferSet->SecondaryIndexBuffer.Indices.Num() > 0
 			&& UseSecondaryMaterialProxy != nullptr
-			&& BufferSet->SecondaryRayTracingGeometry.RayTracingGeometryRHI.IsValid())
+			&& BufferSet->SecondaryRayTracingGeometry.IsValid())
 		{
 			ensure(BufferSet->SecondaryRayTracingGeometry.Initializer.IndexBuffer.IsValid());
 			DrawRayTracingBatch(Context, *BufferSet, BufferSet->SecondaryIndexBuffer, BufferSet->SecondaryRayTracingGeometry, UseSecondaryMaterialProxy, DepthPriority, DynamicPrimitiveUniformBuffer, OutRayTracingInstances);
diff --git a/Engine/Source/Runtime/Renderer/Private/RayTracing/RayTracingDynamicGeometry.cpp b/Engine/Source/Runtime/Renderer/Private/RayTracing/RayTracingDynamicGeometry.cpp
--- a/Engine/Source/Runtime/Renderer/Private/RayTracing/RayTracingDynamicGeometry.cpp
+++ b/Engine/Source/Runtime/Renderer/Private/RayTracing/RayTracingDynamicGeometry.cpp
@@ -314,7 +314,7 @@ void FRayTracingDynamicGeometryCollection::AddDynamicMeshBatchForGeometryUpdate(
 		bRefit = false;
 	}
 
-	if (!Geometry.RayTracingGeometryRHI.IsValid())
+	if (!Geometry.IsValid())
 	{
 		bRefit = false;
 	}
