Add option for dynamic mesh component to use static render path instead of dynamic

diff --git a/Engine/Source/Runtime/GeometryFramework/Private/Components/BaseDynamicMeshComponent.cpp b/Engine/Source/Runtime/GeometryFramework/Private/Components/BaseDynamicMeshComponent.cpp
--- a/Engine/Source/Runtime/GeometryFramework/Private/Components/BaseDynamicMeshComponent.cpp
+++ b/Engine/Source/Runtime/GeometryFramework/Private/Components/BaseDynamicMeshComponent.cpp
@@ -22,7 +22,8 @@ void UBaseDynamicMeshComponent::PostEditChangeProperty(FPropertyChangedEvent& Pr
 	Super::PostEditChangeProperty(PropertyChangedEvent);
 
 	const FName PropName = PropertyChangedEvent.GetPropertyName();
-	if ( (PropName == GET_MEMBER_NAME_CHECKED(UBaseDynamicMeshComponent, bEnableRaytracing))  )
+	if ( (PropName == GET_MEMBER_NAME_CHECKED(UBaseDynamicMeshComponent, bEnableRaytracing)) || 
+		 (PropName == GET_MEMBER_NAME_CHECKED(UBaseDynamicMeshComponent, DrawPath)) )
 	{
 		OnRenderingStateChanged(true);
 	}
@@ -120,6 +121,20 @@ bool UBaseDynamicMeshComponent::GetEnableRaytracing() const
 }
 
 
+void UBaseDynamicMeshComponent::SetMeshDrawPath(EDynamicMeshDrawPath NewDrawPath)
+{
+	if (DrawPath != NewDrawPath)
+	{
+		DrawPath = NewDrawPath;
+		OnRenderingStateChanged(true);
+	}
+}
+
+EDynamicMeshDrawPath UBaseDynamicMeshComponent::GetMeshDrawPath() const
+{
+	return DrawPath;
+}
+
 
 void UBaseDynamicMeshComponent::SetColorOverrideMode(EDynamicMeshComponentColorOverrideMode NewMode)
 {
diff --git a/Engine/Source/Runtime/GeometryFramework/Private/Components/BaseDynamicMeshSceneProxy.cpp b/Engine/Source/Runtime/GeometryFramework/Private/Components/BaseDynamicMeshSceneProxy.cpp
--- a/Engine/Source/Runtime/GeometryFramework/Private/Components/BaseDynamicMeshSceneProxy.cpp
+++ b/Engine/Source/Runtime/GeometryFramework/Private/Components/BaseDynamicMeshSceneProxy.cpp
@@ -16,7 +16,8 @@ FBaseDynamicMeshSceneProxy::FBaseDynamicMeshSceneProxy(UBaseDynamicMeshComponent
 	ParentBaseComponent(Component),
 	ColorSpaceTransformMode(Component->GetVertexColorSpaceTransformMode()),
 	bEnableRaytracing(Component->GetEnableRaytracing()),
-	bEnableViewModeOverrides(Component->GetViewModeOverridesEnabled())
+	bEnableViewModeOverrides(Component->GetViewModeOverridesEnabled()),
+	bPreferStaticDrawPath(Component->GetMeshDrawPath() == EDynamicMeshDrawPath::StaticDraw)
 {
 	if (Component->GetColorOverrideMode() == EDynamicMeshComponentColorOverrideMode::Constant)
 	{
@@ -469,6 +470,105 @@ void FBaseDynamicMeshSceneProxy::DrawBatch(FMeshElementCollector& Collector, con
 	Collector.AddMesh(ViewIndex, Mesh);
 }
 
+
+bool FBaseDynamicMeshSceneProxy::AllowStaticDrawPath(const FSceneView* View) const
+{
+	bool bAllowDebugViews = AllowDebugViewmodes();
+	if (!bAllowDebugViews)
+	{
+		return true;
+	}
+	const FEngineShowFlags& EngineShowFlags = View->Family->EngineShowFlags;
+	bool bWantWireframeOnShaded = ParentBaseComponent->GetEnableWireframeRenderPass();
+	bool bWireframe = EngineShowFlags.Wireframe || bWantWireframeOnShaded;
+	if (bWireframe)
+	{
+		return false;
+	}
+	bool bDrawSimpleCollision = false, bDrawComplexCollision = false;
+	bool bDrawCollisionView = IsCollisionView(EngineShowFlags, bDrawSimpleCollision, bDrawComplexCollision); // check for the full collision views
+	bool bDrawCollisionFlags = EngineShowFlags.Collision && IsCollisionEnabled(); // check for single component collision rendering
+	bool bDrawCollision = bDrawCollisionFlags || bDrawSimpleCollision || bDrawCollisionView;
+	if (bDrawCollision)
+	{
+		return false;
+	}
+	bool bIsSelected = IsSelected();
+	bool bColorOverrides = (bIsSelected && EngineShowFlags.VertexColors) || (ParentBaseComponent->ColorMode != EDynamicMeshComponentColorOverrideMode::None);
+	return !bColorOverrides;
+}
+
+
+void FBaseDynamicMeshSceneProxy::DrawStaticElements(FStaticPrimitiveDrawInterface* PDI)
+{
+	QUICK_SCOPE_CYCLE_COUNTER(STAT_BaseDynamicMeshSceneProxy_DrawStaticElements);
+
+	UMaterialInterface* UseSecondaryMaterial = nullptr;
+	if (ParentBaseComponent->HasSecondaryRenderMaterial())
+	{
+		UseSecondaryMaterial = ParentBaseComponent->GetSecondaryRenderMaterial();
+	}
+	bool bDrawSecondaryBuffers = ParentBaseComponent->GetSecondaryBuffersVisibility();
+
+	ESceneDepthPriorityGroup DepthPriority = SDPG_World;
+
+	TArray<FMeshRenderBufferSet*> Buffers;
+	GetActiveRenderBufferSets(Buffers);
+	PDI->ReserveMemoryForMeshes(Buffers.Num());
+
+	// Draw the mesh.
+	int32 SectionIndexCounter = 0;
+	for (FMeshRenderBufferSet* BufferSet : Buffers)
+	{
+		if (BufferSet->TriangleCount == 0)
+		{
+			continue;
+		}
+
+		UMaterialInterface* UseMaterial = BufferSet->Material;
+		if (ParentBaseComponent->HasOverrideRenderMaterial(0))
+		{
+			UseMaterial = ParentBaseComponent->GetOverrideRenderMaterial(0);
+		}
+		FMaterialRenderProxy* MaterialProxy = UseMaterial->GetRenderProxy();
+
+		// lock buffers so that they aren't modified while we are submitting them
+		FScopeLock BuffersLock(&BufferSet->BuffersLock);
+
+		FMeshBatch MeshBatch;
+
+		FMeshBatchElement& BatchElement = MeshBatch.Elements[0];
+		BatchElement.IndexBuffer = &BufferSet->IndexBuffer;
+		MeshBatch.VertexFactory = &BufferSet->VertexFactory;
+		MeshBatch.MaterialRenderProxy = MaterialProxy;
+
+		BatchElement.PrimitiveUniformBuffer = GetUniformBuffer();
+		BatchElement.NumPrimitives = BufferSet->IndexBuffer.Indices.Num() / 3;
+		BatchElement.FirstIndex = 0;
+		BatchElement.MinVertexIndex = 0;
+		BatchElement.MaxVertexIndex = BufferSet->PositionVertexBuffer.GetNumVertices() - 1;
+		MeshBatch.ReverseCulling = IsLocalToWorldDeterminantNegative();
+		MeshBatch.Type = PT_TriangleList;
+		MeshBatch.DepthPriorityGroup = DepthPriority;
+		MeshBatch.bCanApplyViewModeOverrides = this->bEnableViewModeOverrides;
+		MeshBatch.LODIndex = 0;
+		MeshBatch.SegmentIndex = SectionIndexCounter;
+		MeshBatch.MeshIdInPrimitive = SectionIndexCounter;
+		SectionIndexCounter++;
+
+		MeshBatch.LCI = nullptr; // lightmap cache interface (allowed to be null)
+		MeshBatch.CastShadow = true;
+		MeshBatch.bUseForMaterial = true;
+		MeshBatch.bDitheredLODTransition = false;
+		MeshBatch.bUseForDepthPass = true;
+		MeshBatch.bUseAsOccluder = ShouldUseAsOccluder();
+
+		PDI->DrawMesh(MeshBatch, FLT_MAX);
+	}
+
+}
+
+
 void FBaseDynamicMeshSceneProxy::SetCollisionData()
 {
 #if UE_ENABLE_DEBUG_DRAWING
diff --git a/Engine/Source/Runtime/GeometryFramework/Private/Components/DynamicMeshSceneProxy.h b/Engine/Source/Runtime/GeometryFramework/Private/Components/DynamicMeshSceneProxy.h
--- a/Engine/Source/Runtime/GeometryFramework/Private/Components/DynamicMeshSceneProxy.h
+++ b/Engine/Source/Runtime/GeometryFramework/Private/Components/DynamicMeshSceneProxy.h
@@ -630,7 +630,17 @@ public:
 
 		Result.bDrawRelevance = IsShown(View);
 		Result.bShadowRelevance = IsShadowCast(View);
-		Result.bDynamicRelevance = true;
+
+		bool bUseStaticDrawPath = bPreferStaticDrawPath && AllowStaticDrawPath(View);
+		Result.bDynamicRelevance = !bUseStaticDrawPath;
+		Result.bStaticRelevance = bUseStaticDrawPath;
+#if WITH_EDITOR
+		//only check these in the editor
+		Result.bEditorVisualizeLevelInstanceRelevance = IsEditingLevelInstanceChild();
+		Result.bEditorStaticSelectionRelevance = (IsSelected() || IsHovered());
+#endif
+
+
 		Result.bRenderInMainPass = ShouldRenderInMainPass();
 		Result.bUsesLightingChannels = GetLightingChannelMask() != GetDefaultLightingChannelMask();
 		Result.bTranslucentSelfShadow = bCastVolumetricTranslucentShadow;
diff --git a/Engine/Source/Runtime/GeometryFramework/Public/Components/BaseDynamicMeshComponent.h b/Engine/Source/Runtime/GeometryFramework/Public/Components/BaseDynamicMeshComponent.h
--- a/Engine/Source/Runtime/GeometryFramework/Public/Components/BaseDynamicMeshComponent.h
+++ b/Engine/Source/Runtime/GeometryFramework/Public/Components/BaseDynamicMeshComponent.h
@@ -43,7 +43,7 @@ ENUM_CLASS_FLAGS(EMeshRenderAttributeFlags);
 /**
  * Tangent calculation modes
  */
-UENUM()
+UENUM(BlueprintType)
 enum class EDynamicMeshComponentTangentsMode : uint8
 {
 	/** Tangents are not used/available, proceed accordingly (eg generate arbitrary orthogonal basis) */
@@ -60,7 +60,7 @@ enum class EDynamicMeshComponentTangentsMode : uint8
 /**
  * Color Override Modes
  */
-UENUM()
+UENUM(BlueprintType)
 enum class EDynamicMeshComponentColorOverrideMode : uint8
 {
 	/** No Color Override enabled */
@@ -74,13 +74,26 @@ enum class EDynamicMeshComponentColorOverrideMode : uint8
 };
 
 
+/**
+ * Draw Path to use
+ */
+UENUM(BlueprintType)
+enum class EDynamicMeshDrawPath : uint8
+{
+	/** Use the dynamic draw path, intended for meshes that update on most frames */
+	DynamicDraw = 0,
+	/** Use the static draw path, for meshes that do not change on most frames */
+	StaticDraw = 1
+};
+
+
 /**
  * Color Transform to apply to Vertex Colors when converting from internal DynamicMesh
  * Color attributes (eg Color Overlay stored in FVector4f) to RHI Render Buffers (FColor).
  * 
  * Note that UStaticMesh assumes the Source Mesh colors are Linear and always converts to SRGB.
  */
-UENUM()
+UENUM(BlueprintType)
 enum class EDynamicMeshVertexColorTransformMode : uint8
 {
 	/** Do not apply any color-space transform to Vertex Colors */
@@ -497,6 +510,30 @@ protected:
 
 
 
+
+protected:
+	/**
+	 * Controls whether to use the dynamic or static draw path
+	 */
+	UPROPERTY(EditAnywhere, BlueprintReadWrite,
+				BlueprintSetter=SetMeshDrawPath, BlueprintGetter=GetMeshDrawPath,
+				 Category = "Dynamic Mesh Component|Rendering")
+	EDynamicMeshDrawPath DrawPath = EDynamicMeshDrawPath::DynamicDraw;
+	
+public:
+
+	/**
+	 * Set the mesh rendering to use the static or dynamic draw path
+	 */
+	UFUNCTION(BlueprintCallable, Category = "Dynamic Mesh Component")
+	GEOMETRYFRAMEWORK_API virtual void SetMeshDrawPath(EDynamicMeshDrawPath NewDrawPath);
+
+	/**
+	 * @return Whether the mesh renders using the static or dynamic draw path
+	 */
+	UFUNCTION(BlueprintPure, Category = "Dynamic Mesh Component")
+	GEOMETRYFRAMEWORK_API virtual EDynamicMeshDrawPath GetMeshDrawPath() const;
+
 	//===============================================================================================================
 	// Standard Component interfaces
 	//
diff --git a/Engine/Source/Runtime/GeometryFramework/Public/Components/BaseDynamicMeshSceneProxy.h b/Engine/Source/Runtime/GeometryFramework/Public/Components/BaseDynamicMeshSceneProxy.h
--- a/Engine/Source/Runtime/GeometryFramework/Public/Components/BaseDynamicMeshSceneProxy.h
+++ b/Engine/Source/Runtime/GeometryFramework/Public/Components/BaseDynamicMeshSceneProxy.h
@@ -479,6 +479,9 @@ protected:
 	// Allow view-mode overrides. 
 	bool bEnableViewModeOverrides = true;
 
+	// whether to try to use the static draw instead of dynamic draw path; note we may still use the dynamic path if collision or vertex color rendering is enabled
+	bool bPreferStaticDrawPath = false;
+
 public:
 	GEOMETRYFRAMEWORK_API FBaseDynamicMeshSceneProxy(UBaseDynamicMeshComponent* Component);
 
@@ -968,6 +971,11 @@ protected:
 		FMeshElementCollector& Collector) const;
 public:
 
+	GEOMETRYFRAMEWORK_API virtual void DrawStaticElements(FStaticPrimitiveDrawInterface* PDI) override;
+
+	// Whether to allow use of the static draw path. If false, the dynamic draw path will be used instead.
+	GEOMETRYFRAMEWORK_API virtual bool AllowStaticDrawPath(const FSceneView* View) const;
+
 	/**
 	 * Draw a single-frame FMeshBatch for a FMeshRenderBufferSet
 	 */
@@ -980,6 +988,12 @@ public:
 		int ViewIndex,
 		FDynamicPrimitiveUniformBuffer& DynamicPrimitiveUniformBuffer) const;
 
+
+
+	//
+	// Raytracing APIs
+	//
+
 #if RHI_RAYTRACING
 
 	GEOMETRYFRAMEWORK_API virtual bool IsRayTracingRelevant() const override;
