Add cvar geometry.DynamicMesh.MaxComplexCollisionTriCount to prevent dynamic mesh components with many triangles from auto-generating complex collision.
+ Make complex collision wireframe only render on the base dynamic mesh proxy if the complex collision is provided.
+ add helper to update physics/etc consistently after mesh changes

diff --git a/Engine/Source/Runtime/GeometryFramework/Private/Components/BaseDynamicMeshSceneProxy.cpp b/Engine/Source/Runtime/GeometryFramework/Private/Components/BaseDynamicMeshSceneProxy.cpp
--- a/Engine/Source/Runtime/GeometryFramework/Private/Components/BaseDynamicMeshSceneProxy.cpp
+++ b/Engine/Source/Runtime/GeometryFramework/Private/Components/BaseDynamicMeshSceneProxy.cpp
@@ -352,7 +352,7 @@ void FBaseDynamicMeshSceneProxy::GetCollisionDynamicMeshElements(TArray<FMeshRen
 				bool bDrawComplexWireframeCollision = (EngineShowFlags.Collision && IsCollisionEnabled() && CollisionTraceFlag == ECollisionTraceFlag::CTF_UseComplexAsSimple);
 
 				// If drawing complex collision as solid or wireframe
-				if(bDrawComplexWireframeCollision || (bDrawCollisionView && bDrawComplexCollision))
+				if (bHasComplexMeshData && (bDrawComplexWireframeCollision || (bDrawCollisionView && bDrawComplexCollision)))
 				{
 					bool bDrawWireframe = !bDrawCollisionView;
 
@@ -574,10 +574,16 @@ void FBaseDynamicMeshSceneProxy::SetCollisionData()
 #if UE_ENABLE_DEBUG_DRAWING
 	bHasCollisionData = true;
 	bOwnerIsNull = ParentBaseComponent->GetOwner() == nullptr;
+	bHasComplexMeshData = false;
 	if (UBodySetup* BodySetup = ParentBaseComponent->GetBodySetup())
 	{
 		CollisionTraceFlag = BodySetup->GetCollisionTraceFlag();
 		CachedAggGeom = BodySetup->AggGeom;
+		
+		if (IInterface_CollisionDataProvider* CDP = Cast<IInterface_CollisionDataProvider>(ParentBaseComponent))
+		{
+			bHasComplexMeshData = CDP->ContainsPhysicsTriMeshData(BodySetup->bMeshCollideAll);
+		}
 	}
 	else
 	{
diff --git a/Engine/Source/Runtime/GeometryFramework/Private/Components/DynamicMeshComponent.cpp b/Engine/Source/Runtime/GeometryFramework/Private/Components/DynamicMeshComponent.cpp
--- a/Engine/Source/Runtime/GeometryFramework/Private/Components/DynamicMeshComponent.cpp
+++ b/Engine/Source/Runtime/GeometryFramework/Private/Components/DynamicMeshComponent.cpp
@@ -28,6 +28,16 @@
 
 using namespace UE::Geometry;
 
+
+static TAutoConsoleVariable<int32> CVarDynamicMeshComponent_MaxComplexCollisionTriCount(
+	TEXT("geometry.DynamicMesh.MaxComplexCollisionTriCount"),
+	250000,
+	TEXT("If a DynamicMeshCompnent's UDynamicMesh has a larger triangle count than this value, it will not be passed to the Physics system to be used as Complex Collision geometry. A negative value indicates no limit.")
+);
+
+
+
+
 namespace
 {
 	// probably should be something defined for the whole tool framework...
@@ -1185,16 +1195,7 @@ void UDynamicMeshComponent::OnMeshObjectChanged(UDynamicMesh* ChangedMeshObject,
 		OnMeshChanged.Broadcast();
 	}
 
-	// Rebuild body setup. Should this be deferred until proxy creation? Sometimes multiple changes are emitted...
-	// todo: can possibly skip this in some change situations, eg if only changing attributes
-	if (bDeferCollisionUpdates || bTransientDeferCollisionUpdates )
-	{
-		InvalidatePhysicsData();
-	}
-	else
-	{
-		RebuildPhysicsData();
-	}
+	InternalOnMeshUpdated();
 }
 
 
@@ -1219,15 +1220,7 @@ void UDynamicMeshComponent::SetDynamicMesh(UDynamicMesh* NewMesh)
 	NotifyMeshUpdated();
 	OnMeshChanged.Broadcast();
 
-	// Rebuild physics data
-	if (bDeferCollisionUpdates || bTransientDeferCollisionUpdates)
-	{
-		InvalidatePhysicsData();
-	}
-	else
-	{
-		RebuildPhysicsData();
-	}
+	InternalOnMeshUpdated();
 }
 
 
@@ -1243,6 +1236,18 @@ void UDynamicMeshComponent::OnChildDetached(USceneComponent* ChildComponent)
 	OnChildAttachmentModified.Broadcast(ChildComponent, false);
 }
 
+void UDynamicMeshComponent::InternalOnMeshUpdated()
+{
+	// Rebuild physics data
+	if (bDeferCollisionUpdates || bTransientDeferCollisionUpdates)
+	{
+		InvalidatePhysicsData();
+	}
+	else
+	{
+		RebuildPhysicsData();
+	}
+}
 
 bool UDynamicMeshComponent::GetTriMeshSizeEstimates(struct FTriMeshCollisionDataEstimates& OutTriMeshEstimates, bool bInUseAllTriData) const
 {
@@ -1425,7 +1430,28 @@ bool UDynamicMeshComponent::GetPhysicsTriMeshData(struct FTriMeshCollisionData*
 
 bool UDynamicMeshComponent::ContainsPhysicsTriMeshData(bool InUseAllTriData) const
 {
-	return bEnableComplexCollision && ((MeshObject != nullptr) ? (MeshObject->GetTriangleCount() > 0) : false);
+	if (bEnableComplexCollision && (MeshObject != nullptr))
+	{
+		int32 TriangleCount = MeshObject->GetTriangleCount();
+
+		// if the triangle count is too large, skip building complex collision
+		int32 MaxComplexCollisionTriCount = CVarDynamicMeshComponent_MaxComplexCollisionTriCount.GetValueOnAnyThread();
+		if (MaxComplexCollisionTriCount >= 0 && TriangleCount > MaxComplexCollisionTriCount)
+		{
+			static bool bHavePrintedWarningMessage = false;
+			if (!bHavePrintedWarningMessage)
+			{
+				UE_LOG(LogGeometry, Display, TEXT("Ignoring attempt to build Complex Collision for a DynamicMeshComponent with triangle count larger than %d. Increase the geometry.DynamicMesh.MaxComplexCollisionTriCount value if you are certain you want to build Complex Collision for very large meshes."), MaxComplexCollisionTriCount);
+				bHavePrintedWarningMessage = true;
+			}
+			return false;
+		}
+		if (TriangleCount > 0)
+		{
+			return true;
+		}
+	}
+	return false;
 }
 
 bool UDynamicMeshComponent::WantsNegXTriMesh()
diff --git a/Engine/Source/Runtime/GeometryFramework/Public/Components/BaseDynamicMeshSceneProxy.h b/Engine/Source/Runtime/GeometryFramework/Public/Components/BaseDynamicMeshSceneProxy.h
--- a/Engine/Source/Runtime/GeometryFramework/Public/Components/BaseDynamicMeshSceneProxy.h
+++ b/Engine/Source/Runtime/GeometryFramework/Public/Components/BaseDynamicMeshSceneProxy.h
@@ -1029,6 +1029,8 @@ private:
 	bool bOwnerIsNull = true;
 	/** Whether the collision data has been set up for rendering */
 	bool bHasCollisionData = false;
+	/** Whether a complex collision mesh is available */
+	bool bHasComplexMeshData = false;
 	/** Collision trace flags */
 	ECollisionTraceFlag		CollisionTraceFlag;
 	/** Collision Response of this component */
diff --git a/Engine/Source/Runtime/GeometryFramework/Public/Components/DynamicMeshComponent.h b/Engine/Source/Runtime/GeometryFramework/Public/Components/DynamicMeshComponent.h
--- a/Engine/Source/Runtime/GeometryFramework/Public/Components/DynamicMeshComponent.h
+++ b/Engine/Source/Runtime/GeometryFramework/Public/Components/DynamicMeshComponent.h
@@ -727,4 +727,8 @@ public:
 	GEOMETRYFRAMEWORK_API void SetSceneProxyVerifyUsedMaterials(bool bState);
 
 
+private:
+	// Internal helper to be called when mesh data is updated
+	void InternalOnMeshUpdated();
+
 };
