fix issue where the static draw path was still run when it was not intended to be enabled +
disable fast vertex/index buffer updates when using static draw path for dynamic mesh component, and more consistently lock the dynamic mesh proxy buffersets when accessing them

diff --git a/Engine/Source/Runtime/GeometryFramework/Private/Components/BaseDynamicMeshSceneProxy.cpp b/Engine/Source/Runtime/GeometryFramework/Private/Components/BaseDynamicMeshSceneProxy.cpp
--- a/Engine/Source/Runtime/GeometryFramework/Private/Components/BaseDynamicMeshSceneProxy.cpp
+++ b/Engine/Source/Runtime/GeometryFramework/Private/Components/BaseDynamicMeshSceneProxy.cpp
@@ -503,6 +503,11 @@ void FBaseDynamicMeshSceneProxy::DrawStaticElements(FStaticPrimitiveDrawInterfac
 {
 	QUICK_SCOPE_CYCLE_COUNTER(STAT_BaseDynamicMeshSceneProxy_DrawStaticElements);
 
+	if (!bPreferStaticDrawPath)
+	{
+		return;
+	}
+
 	UMaterialInterface* UseSecondaryMaterial = nullptr;
 	if (ParentBaseComponent->HasSecondaryRenderMaterial())
 	{
diff --git a/Engine/Source/Runtime/GeometryFramework/Private/Components/DynamicMeshComponent.cpp b/Engine/Source/Runtime/GeometryFramework/Private/Components/DynamicMeshComponent.cpp
--- a/Engine/Source/Runtime/GeometryFramework/Private/Components/DynamicMeshComponent.cpp
+++ b/Engine/Source/Runtime/GeometryFramework/Private/Components/DynamicMeshComponent.cpp
@@ -416,7 +416,7 @@ void UDynamicMeshComponent::FastNotifyColorsUpdated()
 	}
 
 	FDynamicMeshSceneProxy* Proxy = GetCurrentSceneProxy();
-	if (Proxy)
+	if (Proxy && AllowFastUpdate())
 	{
 		if (HasTriangleColorFunction() && Proxy->bUsePerTriangleColor == false )
 		{
@@ -462,7 +462,7 @@ void UDynamicMeshComponent::FastNotifyPositionsUpdated(bool bNormals, bool bColo
 	}
 
 	FDynamicMeshSceneProxy* Proxy = GetCurrentSceneProxy();
-	if (Proxy)
+	if (Proxy && AllowFastUpdate())
 	{
 		// calculate bounds while we are updating vertices
 		TFuture<void> UpdateBoundsCalc;
@@ -499,7 +499,7 @@ void UDynamicMeshComponent::FastNotifyVertexAttributesUpdated(bool bNormals, boo
 	}
 
 	FDynamicMeshSceneProxy* Proxy = GetCurrentSceneProxy();
-	if (Proxy && ensure(bNormals || bColors || bUVs) )
+	if (Proxy && ensure(bNormals || bColors || bUVs) && AllowFastUpdate())
 	{
 		GetCurrentSceneProxy()->FastUpdateVertices(false, bNormals, bColors, bUVs);
 		//MarkRenderDynamicDataDirty();
@@ -525,7 +525,7 @@ void UDynamicMeshComponent::FastNotifyVertexAttributesUpdated(EMeshRenderAttribu
 	}
 
 	FDynamicMeshSceneProxy* Proxy = GetCurrentSceneProxy();
-	if (Proxy && ensure(UpdatedAttributes != EMeshRenderAttributeFlags::None))
+	if (Proxy && ensure(UpdatedAttributes != EMeshRenderAttributeFlags::None) && AllowFastUpdate())
 	{
 		bool bPositions = (UpdatedAttributes & EMeshRenderAttributeFlags::Positions) != EMeshRenderAttributeFlags::None;
 
@@ -606,7 +606,7 @@ void UDynamicMeshComponent::FastNotifySecondaryTrianglesChanged()
 	}
 
 	FDynamicMeshSceneProxy* Proxy = GetCurrentSceneProxy();
-	if (Proxy)
+	if (Proxy && AllowFastUpdate())
 	{
 		GetCurrentSceneProxy()->FastUpdateAllIndexBuffers();
 		GetDynamicMesh()->PostRealtimeUpdate();
@@ -632,7 +632,7 @@ void UDynamicMeshComponent::FastNotifyTriangleVerticesUpdated(const TArray<int32
 		((UpdatedAttributes & EMeshRenderAttributeFlags::SecondaryIndexBuffers) != EMeshRenderAttributeFlags::None);
 
 	FDynamicMeshSceneProxy* Proxy = GetCurrentSceneProxy();
-	if (!Proxy)
+	if (!Proxy || !AllowFastUpdate())
 	{
 		ResetProxy();
 	}
@@ -714,7 +714,7 @@ void UDynamicMeshComponent::FastNotifyTriangleVerticesUpdated(const TSet<int32>&
 		((UpdatedAttributes & EMeshRenderAttributeFlags::SecondaryIndexBuffers) != EMeshRenderAttributeFlags::None);
 
 	FDynamicMeshSceneProxy* Proxy = GetCurrentSceneProxy();
-	if (!Proxy)
+	if (!Proxy || !AllowFastUpdate())
 	{
 		ResetProxy();
 	}
@@ -827,7 +827,7 @@ TFuture<bool> UDynamicMeshComponent::FastNotifyTriangleVerticesUpdated_TryPrecom
 	TArray<int32>& UpdateSetsOut,
 	FAxisAlignedBox3d& BoundsOut)
 {
-	if ((!!RenderMeshPostProcessor) || (GetCurrentSceneProxy() == nullptr) || (!Decomposition))
+	if ((!!RenderMeshPostProcessor) || (GetCurrentSceneProxy() == nullptr) || (!Decomposition) || !AllowFastUpdate())
 	{
 		// is there a simpler way to do this? cannot seem to just make a TFuture<bool>...
 		return Async(DynamicMeshComponentAsyncExecTarget, []() { return false; });
@@ -885,7 +885,7 @@ void UDynamicMeshComponent::FastNotifyTriangleVerticesUpdated_ApplyPrecompute(
 	Precompute.Wait();
 
 	bool bPrecomputeOK = Precompute.Get();
-	if (bPrecomputeOK == false || GetCurrentSceneProxy() == nullptr )
+	if (bPrecomputeOK == false || GetCurrentSceneProxy() == nullptr || !AllowFastUpdate())
 	{
 		FastNotifyTriangleVerticesUpdated(Triangles, UpdatedAttributes);
 		return;
diff --git a/Engine/Source/Runtime/GeometryFramework/Private/Components/DynamicMeshSceneProxy.h b/Engine/Source/Runtime/GeometryFramework/Private/Components/DynamicMeshSceneProxy.h
--- a/Engine/Source/Runtime/GeometryFramework/Private/Components/DynamicMeshSceneProxy.h
+++ b/Engine/Source/Runtime/GeometryFramework/Private/Components/DynamicMeshSceneProxy.h
@@ -344,6 +344,7 @@ public:
 			if (ensure(RenderBufferSets.Num() == 1))
 			{
 				FMeshRenderBufferSet* BufferSet = RenderBufferSets[0];
+				FScopeLock BuffersLock(&BufferSet->BuffersLock);
 				if (BufferSet->TriangleCount != Mesh->TriangleCount())
 				{
 					return false;
@@ -361,6 +362,7 @@ public:
 		{
 			for (FMeshRenderBufferSet* BufferSet : RenderBufferSets)
 			{
+				FScopeLock BuffersLock(&BufferSet->BuffersLock);
 				if (ensure(BufferSet->Triangles))
 				{
 					int NumTriangles = BufferSet->Triangles->Num();
@@ -418,6 +420,7 @@ public:
 			if (ensure(RenderBufferSets.Num() == 1))
 			{
 				FMeshRenderBufferSet* Buffers = RenderBufferSets[0];
+				FScopeLock BuffersLock(&Buffers->BuffersLock);
 				if (bPositions || bNormals || bColors)
 				{
 					UpdateVertexBuffersFromOverlays(Buffers, Mesh,
@@ -444,6 +447,7 @@ public:
 			ParallelFor(RenderBufferSets.Num(), [&](int i)
 			{
 				FMeshRenderBufferSet* Buffers = RenderBufferSets[i];
+				FScopeLock BuffersLock(&Buffers->BuffersLock);
 				if (Buffers->TriangleCount == 0)
 				{
 					return;
@@ -521,11 +525,16 @@ public:
 		ParallelFor(WhichBuffers.Num(), [&](int idx)
 		{
 			int32 BufferIndex = WhichBuffers[idx];
-			if ( RenderBufferSets.IsValidIndex(BufferIndex) == false || RenderBufferSets[BufferIndex]->TriangleCount == 0)
+			if ( RenderBufferSets.IsValidIndex(BufferIndex) == false)
 			{
 				return;
 			}
 			FMeshRenderBufferSet* Buffers = RenderBufferSets[BufferIndex];
+			FScopeLock BuffersLock(&Buffers->BuffersLock);
+			if (Buffers->TriangleCount == 0)
+			{
+				return;
+			}
 			if (ensure(Buffers->Triangles.IsSet()))
 			{
 				if (bPositions || bNormals || bColors)
diff --git a/Engine/Source/Runtime/GeometryFramework/Public/Components/BaseDynamicMeshComponent.h b/Engine/Source/Runtime/GeometryFramework/Public/Components/BaseDynamicMeshComponent.h
--- a/Engine/Source/Runtime/GeometryFramework/Public/Components/BaseDynamicMeshComponent.h
+++ b/Engine/Source/Runtime/GeometryFramework/Public/Components/BaseDynamicMeshComponent.h
@@ -519,6 +519,14 @@ protected:
 				BlueprintSetter=SetMeshDrawPath, BlueprintGetter=GetMeshDrawPath,
 				 Category = "Dynamic Mesh Component|Rendering")
 	EDynamicMeshDrawPath DrawPath = EDynamicMeshDrawPath::DynamicDraw;
+
+	// Whether the fast update paths will be used for the mesh.
+	// If false, fast update methods can still be called, but will fall back to recreating the render proxy
+	virtual bool AllowFastUpdate()
+	{
+		// note: fast update is not compatible with static draw path
+		return DrawPath == EDynamicMeshDrawPath::DynamicDraw;
+	}
 	
 public:
 
