make dynamic mesh collision updates also update the preview collision visualization

diff --git a/Engine/Source/Runtime/GeometryFramework/Private/Components/BaseDynamicMeshSceneProxy.cpp b/Engine/Source/Runtime/GeometryFramework/Private/Components/BaseDynamicMeshSceneProxy.cpp
--- a/Engine/Source/Runtime/GeometryFramework/Private/Components/BaseDynamicMeshSceneProxy.cpp
+++ b/Engine/Source/Runtime/GeometryFramework/Private/Components/BaseDynamicMeshSceneProxy.cpp
@@ -155,6 +155,7 @@ bool FBaseDynamicMeshSceneProxy::IsCollisionView(const FEngineShowFlags& EngineS
 
 #if UE_ENABLE_DEBUG_DRAWING
 	// If in a 'collision view' and collision is enabled
+	FScopeLock Lock(&CachedCollisionLock);
 	if (bHasCollisionData && bDrawCollisionView && IsCollisionEnabled())
 	{
 		// See if we have a response to the interested channel
@@ -331,6 +332,8 @@ void FBaseDynamicMeshSceneProxy::GetCollisionDynamicMeshElements(TArray<FMeshRen
 	const TArray<const FSceneView*>& Views, uint32 VisibilityMap, FMeshElementCollector& Collector) const
 {
 #if UE_ENABLE_DEBUG_DRAWING
+	FScopeLock Lock(&CachedCollisionLock);
+
 	if (!bHasCollisionData)
 	{
 		return;
@@ -577,6 +580,7 @@ void FBaseDynamicMeshSceneProxy::DrawStaticElements(FStaticPrimitiveDrawInterfac
 void FBaseDynamicMeshSceneProxy::SetCollisionData()
 {
 #if UE_ENABLE_DEBUG_DRAWING
+	FScopeLock Lock(&CachedCollisionLock);
 	bHasCollisionData = true;
 	bOwnerIsNull = ParentBaseComponent->GetOwner() == nullptr;
 	bHasComplexMeshData = false;
diff --git a/Engine/Source/Runtime/GeometryFramework/Private/Components/DynamicMeshComponent.cpp b/Engine/Source/Runtime/GeometryFramework/Private/Components/DynamicMeshComponent.cpp
--- a/Engine/Source/Runtime/GeometryFramework/Private/Components/DynamicMeshComponent.cpp
+++ b/Engine/Source/Runtime/GeometryFramework/Private/Components/DynamicMeshComponent.cpp
@@ -1581,6 +1581,11 @@ void UDynamicMeshComponent::RebuildPhysicsData()
 
 		bCollisionUpdatePending = false;
 	}
+
+	if (FDynamicMeshSceneProxy* Proxy = GetCurrentSceneProxy())
+	{
+		Proxy->SetCollisionData();
+	}
 }
 
 void UDynamicMeshComponent::FinishPhysicsAsyncCook(bool bSuccess, UBodySetup* FinishedBodySetup)
diff --git a/Engine/Source/Runtime/GeometryFramework/Public/Components/BaseDynamicMeshSceneProxy.h b/Engine/Source/Runtime/GeometryFramework/Public/Components/BaseDynamicMeshSceneProxy.h
--- a/Engine/Source/Runtime/GeometryFramework/Public/Components/BaseDynamicMeshSceneProxy.h
+++ b/Engine/Source/Runtime/GeometryFramework/Public/Components/BaseDynamicMeshSceneProxy.h
@@ -1038,6 +1038,9 @@ private:
 	/** Cached AggGeom holding the collision shapes to render */
 	FKAggregateGeom CachedAggGeom;
 
+	// Control access to collision data for debug rendering
+	mutable FCriticalSection CachedCollisionLock;
+
 #endif
 
 	GEOMETRYFRAMEWORK_API bool IsCollisionView(const FEngineShowFlags& EngineShowFlags, bool& bDrawSimpleCollision, bool& bDrawComplexCollision) const;
