Initial Lumen support for DynamicMeshComponent

diff --git a/Engine/Source/Runtime/GeometryFramework/Private/Components/BaseDynamicMeshComponent.cpp b/Engine/Source/Runtime/GeometryFramework/Private/Components/BaseDynamicMeshComponent.cpp
--- a/Engine/Source/Runtime/GeometryFramework/Private/Components/BaseDynamicMeshComponent.cpp
+++ b/Engine/Source/Runtime/GeometryFramework/Private/Components/BaseDynamicMeshComponent.cpp
@@ -31,6 +31,10 @@ void UBaseDynamicMeshComponent::PostEditChangeProperty(FPropertyChangedEvent& Pr
 	{
 		OnRenderingStateChanged(false);
 	}
+	else if ( (PropName == GET_MEMBER_NAME_CHECKED(UBaseDynamicMeshComponent, DistanceFieldMode))  )
+	{
+		OnNewDistanceFieldMode();
+	}
 }
 #endif
 
@@ -120,6 +124,26 @@ bool UBaseDynamicMeshComponent::GetEnableRaytracing() const
 	return bEnableRaytracing;
 }
 
+void UBaseDynamicMeshComponent::SetDistanceFieldMode(EDynamicMeshComponentDistanceFieldMode NewDistFieldMode)
+{
+	if (DistanceFieldMode != NewDistFieldMode)
+	{
+		DistanceFieldMode = NewDistFieldMode;
+		OnNewDistanceFieldMode();
+	}
+}
+
+EDynamicMeshComponentDistanceFieldMode UBaseDynamicMeshComponent::GetDistanceFieldMode() const
+{
+	return DistanceFieldMode;
+}
+
+void UBaseDynamicMeshComponent::OnNewDistanceFieldMode()
+{
+	// no-op in base class, does not directly support distance field
+}
+
+
 
 void UBaseDynamicMeshComponent::SetMeshDrawPath(EDynamicMeshDrawPath NewDrawPath)
 {
diff --git a/Engine/Source/Runtime/GeometryFramework/Private/Components/BaseDynamicMeshSceneProxy.cpp b/Engine/Source/Runtime/GeometryFramework/Private/Components/BaseDynamicMeshSceneProxy.cpp
--- a/Engine/Source/Runtime/GeometryFramework/Private/Components/BaseDynamicMeshSceneProxy.cpp
+++ b/Engine/Source/Runtime/GeometryFramework/Private/Components/BaseDynamicMeshSceneProxy.cpp
@@ -10,6 +10,43 @@
 #include "SceneInterface.h"
 #include "SceneManagement.h"
 #include "Engine/Engine.h"		// for GEngine definition
+#include "MeshCardRepresentation.h"
+#include "MeshCardBuild.h"
+#include "DistanceFieldAtlas.h"
+#include "DataDrivenShaderPlatformInfo.h"
+
+#include "Implicit/SweepingMeshSDF.h"
+#include "DynamicMesh/DynamicMeshAABBTree3.h"
+#include "Spatial/FastWinding.h"
+#include "HAL/IConsoleManager.h"
+
+
+static TAutoConsoleVariable<bool> CVarDynamicMeshComponent_AllowDistanceFieldGeneration(
+	TEXT("geometry.DynamicMesh.AllowDistanceFieldGeneration"),
+	1,
+	TEXT("Whether to allow distance field generation for dynamic mesh components")
+);
+
+static TAutoConsoleVariable<bool> CVarDynamicMeshComponent_AllowMeshCardGeneration(
+	TEXT("geometry.DynamicMesh.AllowMeshCardGeneration"),
+	1,
+	TEXT("Whether to allow mesh card generation for dynamic mesh components")
+);
+
+
+namespace UE::DynamicMesh
+{
+	static bool AllowDistanceFieldGeneration()
+	{
+		// we disallow distance fields on integrated devices to match FSceneRenderer::ShouldPrepareDistanceFieldScene, which notes that they are too likely to hang/fail on the associated large allocations
+		return CVarDynamicMeshComponent_AllowDistanceFieldGeneration.GetValueOnAnyThread() && DoesProjectSupportDistanceFields() && !GRHIDeviceIsIntegrated;
+	}
+
+	static bool AllowLumenCardGeneration()
+	{
+		return CVarDynamicMeshComponent_AllowMeshCardGeneration.GetValueOnAnyThread() && FDataDrivenShaderPlatformInfo::GetSupportsLumenGI(GetFeatureLevelShaderPlatform(GMaxRHIFeatureLevel));
+	}
+}
 
 FBaseDynamicMeshSceneProxy::FBaseDynamicMeshSceneProxy(UBaseDynamicMeshComponent* Component)
 	: FPrimitiveSceneProxy(Component),
@@ -28,6 +65,17 @@ FBaseDynamicMeshSceneProxy::FBaseDynamicMeshSceneProxy(UBaseDynamicMeshComponent
 	bUsePerTriangleNormals = Component->GetFlatShadingEnabled();
 	
 	SetCollisionData();
+
+	FMaterialRelevance MaterialRelevance = Component->GetMaterialRelevance(GetScene().GetFeatureLevel());
+	bOpaqueOrMasked = MaterialRelevance.bOpaque;
+
+	// set initial distance field flags based on whether we will have one, after its async build
+	bool bWillHaveDistanceField = Component->GetDistanceFieldMode() != EDynamicMeshComponentDistanceFieldMode::NoDistanceField 
+									&& UE::DynamicMesh::AllowDistanceFieldGeneration();
+	bSupportsDistanceFieldRepresentation = bWillHaveDistanceField;
+	bAffectDistanceFieldLighting = bWillHaveDistanceField;
+	// note whether lumen is enabled will depend on the distance field flags (in some cases)
+	UpdateVisibleInLumenScene();
 }
 
 FBaseDynamicMeshSceneProxy::~FBaseDynamicMeshSceneProxy()
@@ -733,3 +781,533 @@ void FBaseDynamicMeshSceneProxy::DrawRayTracingBatch(FRayTracingMaterialGatherin
 }
 
 #endif // RHI_RAYTRACING
+
+
+
+
+
+
+const FCardRepresentationData* FBaseDynamicMeshSceneProxy::GetMeshCardRepresentation() const
+{
+	if (MeshCards.IsValid() && bMeshCardsValid)
+	{
+		return MeshCards.Get();
+	}
+	return nullptr;
+}
+
+namespace UE::Local
+{
+	// Same as LumenMeshCards::GetAxisAlignedDirection
+	FVector3f GetAxisAlignedDirection(int32 AxisAlignedDirectionIndex, int32& AxisIndex)
+	{
+		AxisIndex = AxisAlignedDirectionIndex / 2;
+		FVector3f Direction(0.0f, 0.0f, 0.0f);
+		Direction[AxisIndex] = AxisAlignedDirectionIndex & 1 ? 1.0f : -1.0f;
+		return Direction;
+	}
+}
+
+void FBaseDynamicMeshSceneProxy::UpdateLumenCardsFromBounds()
+{
+	bMeshCardsValid = false;
+	if (!bVisibleInLumenScene || !UE::DynamicMesh::AllowLumenCardGeneration())
+	{
+		MeshCards.Reset();
+		return;
+	}
+
+	FBox Box = ParentBaseComponent->GetLocalBounds().GetBox();
+
+	if (MeshCards.IsValid() == false)
+	{
+		MeshCards = MakePimpl<FCardRepresentationData>();
+	}
+
+	*MeshCards = FCardRepresentationData();		 // increments ID
+	FMeshCardsBuildData& CardData = MeshCards->MeshCardsBuildData;
+
+	CardData.Bounds = Box;
+
+
+	struct FCardDirection
+	{
+		int DirectionIndex;
+		FVector AxisZ;
+		FVector AxisX;
+		FVector AxisY;
+		int AxisZIndex;
+	};
+	TArray<FCardDirection> CardDirections;
+	for (int32 DirectionIndex = 0; DirectionIndex < 6; ++DirectionIndex)
+	{
+		FCardDirection Direction;
+		Direction.DirectionIndex = DirectionIndex;
+		Direction.AxisZ = (FVector)UE::Local::GetAxisAlignedDirection(DirectionIndex, Direction.AxisZIndex);
+		Direction.AxisZ.FindBestAxisVectors(Direction.AxisX, Direction.AxisY);
+		Direction.AxisX = FVector::CrossProduct(Direction.AxisZ, Direction.AxisY);
+		Direction.AxisX.Normalize();
+		CardDirections.Add(Direction);
+	}
+
+	FVector3d Center = Box.GetCenter();
+	FVector3d Extents = Box.GetExtent();
+	float CardOffset = 5.0;
+
+	CardData.CardBuildData.SetNum(CardDirections.Num());
+	for (int32 CardIndex = 0; CardIndex < CardDirections.Num(); ++CardIndex)
+	{
+		FCardDirection Direction = CardDirections[CardIndex];
+		FLumenCardOBBf OBB;
+		OBB.AxisZ = (FVector3f)Direction.AxisZ;
+		OBB.AxisX = (FVector3f)Direction.AxisX;
+		OBB.AxisY = (FVector3f)Direction.AxisY;
+	
+		// project 3D mesh extents onto the specific axes of this CardOBB  (this just reshuffles them but the combinatorics are messy)
+		double ExtentX = FMathd::Abs(Direction.AxisX.Dot( Extents ));
+		double ExtentY = FMathd::Abs(Direction.AxisY.Dot( Extents ));
+		double ExtentZ = FMathd::Abs(Direction.AxisZ.Dot( Extents ));
+
+		// Translate the box along the AxisZ axis so the box center is at the middle of the axis-face 
+		FVector3d LocalCenter = Center;
+		LocalCenter +=  ExtentZ * Direction.AxisZ;
+
+		// hardcoding the card box to cover half the mesh bounds along Z (and full mesh box along X and Y)
+		double CardExtentZ = ExtentZ * 0.5;	
+
+		// shift the card box center so that the +Z face lies on the mesh box face, then bump it forward a bit
+		LocalCenter += (-CardExtentZ + CardOffset) * Direction.AxisZ;
+
+		// set up the box for the card
+		OBB.Extent = (FVector3f)FVector3d(ExtentX, ExtentY, CardExtentZ + CardOffset*0.5);
+		OBB.Origin = (FVector3f)LocalCenter;
+
+		CardData.CardBuildData[CardIndex].OBB = OBB;
+		CardData.CardBuildData[CardIndex].AxisAlignedDirectionIndex = Direction.DirectionIndex;
+	}
+
+	bMeshCardsValid = true;
+}
+
+
+
+void FBaseDynamicMeshSceneProxy::GetDistanceFieldAtlasData(const class FDistanceFieldVolumeData*& OutDistanceFieldData, float& SelfShadowBias) const
+{
+	if (DistanceField.IsValid() && bDistanceFieldValid)
+	{
+		OutDistanceFieldData = DistanceField.Get();
+		SelfShadowBias = 0.0;
+	}
+	else
+	{
+		OutDistanceFieldData = nullptr;
+		SelfShadowBias = 0.0;
+	}
+}
+
+
+void FBaseDynamicMeshSceneProxy::GetDistanceFieldInstanceData(TArray<FRenderTransform>& InstanceLocalToPrimitiveTransforms) const
+{
+	check(InstanceLocalToPrimitiveTransforms.IsEmpty());
+	if (DistanceField.IsValid() && bDistanceFieldValid)
+	{
+		InstanceLocalToPrimitiveTransforms.Add(FRenderTransform::Identity);
+	}
+}
+
+bool FBaseDynamicMeshSceneProxy::HasDistanceFieldRepresentation() const 
+{
+	return CastsDynamicShadow() && AffectsDistanceFieldLighting() && bDistanceFieldValid && DistanceField.IsValid();
+}
+
+bool FBaseDynamicMeshSceneProxy::HasDynamicIndirectShadowCasterRepresentation() const
+{
+	return bCastsDynamicIndirectShadow && FBaseDynamicMeshSceneProxy::HasDistanceFieldRepresentation();
+}
+
+
+
+
+
+
+static int32 ComputeLinearVoxelIndex(FIntVector VoxelCoordinate, FIntVector VolumeDimensions)
+{
+	return (VoxelCoordinate.Z * VolumeDimensions.Y + VoxelCoordinate.Y) * VolumeDimensions.X + VoxelCoordinate.X;
+}
+
+static bool DynamicMesh_GenerateSignedDistanceFieldVolumeData(
+	const FDynamicMesh3& Mesh,
+	float DistanceFieldResolutionScale,
+	bool bGenerateAsIfTwoSided,
+	FDistanceFieldVolumeData& VolumeDataOut,
+	FProgressCancel& Progress)
+{
+	TRACE_CPUPROFILER_EVENT_SCOPE(DynamicMesh_GenerateSignedDistanceFieldVolumeData);
+
+	if (!UE::DynamicMesh::AllowDistanceFieldGeneration())
+	{
+		return false;
+	}
+
+	if ( DistanceFieldResolutionScale <= 0 )
+	{
+		return false;
+	}
+
+	const double StartTime = FPlatformTime::Seconds();
+
+	UE::Geometry::FDynamicMeshAABBTree3 Spatial(&Mesh, true);
+	if ( Progress.Cancelled() ) { return false; }
+	UE::Geometry::FAxisAlignedBox3d MeshBounds = Spatial.GetBoundingBox();
+	UE::Geometry::TFastWindingTree<FDynamicMesh3> WindingTree(&Spatial, true);
+	if ( Progress.Cancelled() ) { return false; }
+
+	static const auto CVar = IConsoleManager::Get().FindTConsoleVariableDataInt(TEXT("r.DistanceFields.MaxPerMeshResolution"));
+	const int32 PerMeshMax = CVar->GetValueOnAnyThread();
+
+	// Meshes with explicit artist-specified scale can go higher
+	const int32 MaxNumBlocksOneDim = FMath::Min<int32>(FMath::DivideAndRoundNearest(DistanceFieldResolutionScale <= 1 ? PerMeshMax / 2 : PerMeshMax, DistanceField::UniqueDataBrickSize), DistanceField::MaxIndirectionDimension - 1);
+
+	static const auto CVarDensity = IConsoleManager::Get().FindTConsoleVariableDataFloat(TEXT("r.DistanceFields.DefaultVoxelDensity"));
+	const float VoxelDensity = CVarDensity->GetValueOnAnyThread();
+
+	const float NumVoxelsPerLocalSpaceUnit = VoxelDensity * DistanceFieldResolutionScale;
+	FBox3f LocalSpaceMeshBounds = (FBox3f)MeshBounds;
+
+	// Make sure the mesh bounding box has positive extents to handle planes
+	{
+		FVector3f MeshBoundsCenter = LocalSpaceMeshBounds.GetCenter();
+		FVector3f MeshBoundsExtent = FVector3f::Max(LocalSpaceMeshBounds.GetExtent(), FVector3f(1.0f, 1.0f, 1.0f));
+		LocalSpaceMeshBounds.Min = MeshBoundsCenter - MeshBoundsExtent;
+		LocalSpaceMeshBounds.Max = MeshBoundsCenter + MeshBoundsExtent;
+	}
+
+	// We sample on voxel corners and use central differencing for gradients, so a box mesh using two-sided materials whose vertices lie on LocalSpaceMeshBounds produces a zero gradient on intersection
+	// Expand the mesh bounds by a fraction of a voxel to allow room for a pullback on the hit location for computing the gradient.
+	// Only expand for two sided meshes as this adds significant Mesh SDF tracing cost
+	if (bGenerateAsIfTwoSided)
+	{
+		const FVector3f DesiredDimensions = FVector3f(LocalSpaceMeshBounds.GetSize() * FVector3f(NumVoxelsPerLocalSpaceUnit / (float)DistanceField::UniqueDataBrickSize));
+		const FIntVector Mip0IndirectionDimensions = FIntVector(
+			FMath::Clamp(FMath::RoundToInt(DesiredDimensions.X), 1, MaxNumBlocksOneDim),
+			FMath::Clamp(FMath::RoundToInt(DesiredDimensions.Y), 1, MaxNumBlocksOneDim),
+			FMath::Clamp(FMath::RoundToInt(DesiredDimensions.Z), 1, MaxNumBlocksOneDim));
+
+		const float CentralDifferencingExpandInVoxels = .25f;
+		const FVector3f TexelObjectSpaceSize = LocalSpaceMeshBounds.GetSize() / FVector3f(Mip0IndirectionDimensions * DistanceField::UniqueDataBrickSize - FIntVector(2 * CentralDifferencingExpandInVoxels));
+		LocalSpaceMeshBounds = LocalSpaceMeshBounds.ExpandBy(TexelObjectSpaceSize);
+	}
+
+	// The tracing shader uses a Volume space that is normalized by the maximum extent, to keep Volume space within [-1, 1], we must match that behavior when encoding
+	const float LocalToVolumeScale = 1.0f / LocalSpaceMeshBounds.GetExtent().GetMax();
+
+	const FVector3f DesiredDimensions = FVector3f(LocalSpaceMeshBounds.GetSize() * FVector3f(NumVoxelsPerLocalSpaceUnit / (float)DistanceField::UniqueDataBrickSize));
+	const FIntVector Mip0IndirectionDimensions = FIntVector(
+		FMath::Clamp(FMath::RoundToInt(DesiredDimensions.X), 1, MaxNumBlocksOneDim),
+		FMath::Clamp(FMath::RoundToInt(DesiredDimensions.Y), 1, MaxNumBlocksOneDim),
+		FMath::Clamp(FMath::RoundToInt(DesiredDimensions.Z), 1, MaxNumBlocksOneDim));
+
+	TArray<uint8> StreamableMipData;
+
+	struct FDistanceFieldBrick
+	{
+		FDistanceFieldBrick(
+			float InLocalSpaceTraceDistance,
+			FBox3f InVolumeBounds,
+			float InLocalToVolumeScale,
+			FVector2f InDistanceFieldToVolumeScaleBias,
+			FIntVector InBrickCoordinate,
+			FIntVector InIndirectionSize)
+			:
+			LocalSpaceTraceDistance(InLocalSpaceTraceDistance),
+			VolumeBounds(InVolumeBounds),
+			LocalToVolumeScale(InLocalToVolumeScale),
+			DistanceFieldToVolumeScaleBias(InDistanceFieldToVolumeScaleBias),
+			BrickCoordinate(InBrickCoordinate),
+			IndirectionSize(InIndirectionSize),
+			BrickMaxDistance(MIN_uint8),
+			BrickMinDistance(MAX_uint8)
+		{}
+
+		float LocalSpaceTraceDistance;
+		FBox3f VolumeBounds;
+		float LocalToVolumeScale;
+		FVector2f DistanceFieldToVolumeScaleBias;
+		FIntVector BrickCoordinate;
+		FIntVector IndirectionSize;
+
+		// Output
+		uint8 BrickMaxDistance;
+		uint8 BrickMinDistance;
+		TArray<uint8> DistanceFieldVolume;
+	};
+
+
+	for (int32 MipIndex = 0; MipIndex < DistanceField::NumMips; MipIndex++)
+	{
+		if ( Progress.Cancelled() ) { return false; }
+
+		const FIntVector IndirectionDimensions = FIntVector(
+			FMath::DivideAndRoundUp(Mip0IndirectionDimensions.X, 1 << MipIndex),
+			FMath::DivideAndRoundUp(Mip0IndirectionDimensions.Y, 1 << MipIndex),
+			FMath::DivideAndRoundUp(Mip0IndirectionDimensions.Z, 1 << MipIndex));
+
+		// Expand to guarantee one voxel border for gradient reconstruction using bilinear filtering
+		const FVector3f TexelObjectSpaceSize = LocalSpaceMeshBounds.GetSize() / FVector3f(IndirectionDimensions * DistanceField::UniqueDataBrickSize - FIntVector(2 * DistanceField::MeshDistanceFieldObjectBorder));
+		const FBox3f DistanceFieldVolumeBounds = LocalSpaceMeshBounds.ExpandBy(TexelObjectSpaceSize);
+
+		const FVector3f IndirectionVoxelSize = DistanceFieldVolumeBounds.GetSize() / FVector3f(IndirectionDimensions);
+		const float IndirectionVoxelRadius = IndirectionVoxelSize.Size();
+
+		const FVector3f VolumeSpaceDistanceFieldVoxelSize = IndirectionVoxelSize * LocalToVolumeScale / FVector3f(DistanceField::UniqueDataBrickSize);
+		const float MaxDistanceForEncoding = VolumeSpaceDistanceFieldVoxelSize.Size() * DistanceField::BandSizeInVoxels;
+		const float LocalSpaceTraceDistance = MaxDistanceForEncoding / LocalToVolumeScale;
+		const FVector2f DistanceFieldToVolumeScaleBias(2.0f * MaxDistanceForEncoding, -MaxDistanceForEncoding);
+
+		TArray<FDistanceFieldBrick> BricksToCompute;
+		BricksToCompute.Reserve(IndirectionDimensions.X * IndirectionDimensions.Y * IndirectionDimensions.Z / 8);
+		for (int32 ZIndex = 0; ZIndex < IndirectionDimensions.Z; ZIndex++)
+		{
+			for (int32 YIndex = 0; YIndex < IndirectionDimensions.Y; YIndex++)
+			{
+				for (int32 XIndex = 0; XIndex < IndirectionDimensions.X; XIndex++)
+				{
+					BricksToCompute.Emplace(
+						LocalSpaceTraceDistance,
+						DistanceFieldVolumeBounds,
+						LocalToVolumeScale,
+						DistanceFieldToVolumeScaleBias,
+						FIntVector(XIndex, YIndex, ZIndex),
+						IndirectionDimensions);
+				}
+			}
+		}
+
+		if ( Progress.Cancelled() ) { return false; }
+
+		// compute bricks now
+		for ( FDistanceFieldBrick& Brick : BricksToCompute )
+		{
+			const FVector3f BrickIndirectionVoxelSize = Brick.VolumeBounds.GetSize() / FVector3f(Brick.IndirectionSize);
+			const FVector3f DistanceFieldVoxelSize = BrickIndirectionVoxelSize / FVector3f(DistanceField::UniqueDataBrickSize);
+			const FVector3f BrickMinPosition = Brick.VolumeBounds.Min + FVector3f(Brick.BrickCoordinate) * BrickIndirectionVoxelSize;
+
+			Brick.DistanceFieldVolume.Empty(DistanceField::BrickSize * DistanceField::BrickSize * DistanceField::BrickSize);
+			Brick.DistanceFieldVolume.AddZeroed(DistanceField::BrickSize * DistanceField::BrickSize * DistanceField::BrickSize);
+
+			for (int32 ZIndex = 0; ZIndex < DistanceField::BrickSize; ZIndex++)
+			{
+				for (int32 YIndex = 0; YIndex < DistanceField::BrickSize; YIndex++)
+				{
+					if ( Progress.Cancelled() ) { return false; }
+
+					for (int32 XIndex = 0; XIndex < DistanceField::BrickSize; XIndex++)
+					{
+						const FVector3f VoxelPosition = FVector3f(XIndex, YIndex, ZIndex) * DistanceFieldVoxelSize + BrickMinPosition;
+						const int32 Index = (ZIndex * DistanceField::BrickSize * DistanceField::BrickSize + YIndex * DistanceField::BrickSize + XIndex);
+
+						float MinLocalSpaceDistance = LocalSpaceTraceDistance;
+
+						double NearestDistSqr = 0;
+						int32 NearestTriangleID = Spatial.FindNearestTriangle((FVector3d)VoxelPosition, NearestDistSqr, 
+							UE::Geometry::IMeshSpatial::FQueryOptions(LocalSpaceTraceDistance));
+						if (NearestTriangleID != IndexConstants::InvalidID)
+						{
+							const float ClosestDistance = FMath::Sqrt(NearestDistSqr);
+							MinLocalSpaceDistance = FMath::Min(MinLocalSpaceDistance, ClosestDistance);
+
+							// found closest point within search radius
+							double IsoThreshold = 0.5;
+							bool bInside = WindingTree.IsInside((FVector3d)VoxelPosition, 0.5);
+							if ( bInside )
+							{
+								MinLocalSpaceDistance *= -1;
+							}
+						}
+						else
+						{
+							// no closest point...
+							MinLocalSpaceDistance = LocalSpaceTraceDistance;
+						}
+
+						// Transform to the tracing shader's Volume space
+						const float VolumeSpaceDistance = MinLocalSpaceDistance * LocalToVolumeScale;
+						// Transform to the Distance Field texture's space
+						const float RescaledDistance = (VolumeSpaceDistance - DistanceFieldToVolumeScaleBias.Y) / DistanceFieldToVolumeScaleBias.X;
+						check(DistanceField::DistanceFieldFormat == PF_G8);
+						const uint8 QuantizedDistance = FMath::Clamp<int32>(FMath::FloorToInt(RescaledDistance * 255.0f + .5f), 0, 255);
+						Brick.DistanceFieldVolume[Index] = QuantizedDistance;
+						Brick.BrickMaxDistance = FMath::Max(Brick.BrickMaxDistance, QuantizedDistance);
+						Brick.BrickMinDistance = FMath::Min(Brick.BrickMinDistance, QuantizedDistance);
+
+					} // X iteration 
+				} // Y iteration
+			} // Z iteration
+
+			
+		}  // Bricks iteration
+
+
+		FSparseDistanceFieldMip& OutMip = VolumeDataOut.Mips[MipIndex];
+		TArray<uint32> IndirectionTable;
+		IndirectionTable.Empty(IndirectionDimensions.X * IndirectionDimensions.Y * IndirectionDimensions.Z);
+		IndirectionTable.AddUninitialized(IndirectionDimensions.X * IndirectionDimensions.Y * IndirectionDimensions.Z);
+
+		for (int32 i = 0; i < IndirectionTable.Num(); i++)
+		{
+			IndirectionTable[i] = DistanceField::InvalidBrickIndex;
+		} 
+
+		TArray<FDistanceFieldBrick*> ValidBricks;
+		ValidBricks.Reserve(BricksToCompute.Num());
+
+		for (int32 k = 0; k < BricksToCompute.Num(); k++)
+		{
+			const FDistanceFieldBrick& ComputedBrick = BricksToCompute[k];
+			if (ComputedBrick.BrickMinDistance < MAX_uint8 && ComputedBrick.BrickMaxDistance > MIN_uint8)
+			{
+				ValidBricks.Add(&BricksToCompute[k]);
+			}
+		}
+
+		const uint32 NumBricks = ValidBricks.Num();
+		const uint32 BrickSizeBytes = DistanceField::BrickSize * DistanceField::BrickSize * DistanceField::BrickSize * GPixelFormats[DistanceField::DistanceFieldFormat].BlockBytes;
+
+		TArray<uint8> DistanceFieldBrickData;
+		DistanceFieldBrickData.Empty(BrickSizeBytes * NumBricks);
+		DistanceFieldBrickData.AddUninitialized(BrickSizeBytes * NumBricks);
+
+		if ( Progress.Cancelled() ) { return false; }
+
+		for (int32 BrickIndex = 0; BrickIndex < ValidBricks.Num(); BrickIndex++)
+		{
+			const FDistanceFieldBrick& Brick = *ValidBricks[BrickIndex];
+			const int32 IndirectionIndex = ComputeLinearVoxelIndex(Brick.BrickCoordinate, IndirectionDimensions);
+			IndirectionTable[IndirectionIndex] = BrickIndex;
+
+			check(BrickSizeBytes == Brick.DistanceFieldVolume.Num() * Brick.DistanceFieldVolume.GetTypeSize());
+			FPlatformMemory::Memcpy(&DistanceFieldBrickData[BrickIndex * BrickSizeBytes], Brick.DistanceFieldVolume.GetData(), Brick.DistanceFieldVolume.Num() * Brick.DistanceFieldVolume.GetTypeSize());
+		}
+
+		const int32 IndirectionTableBytes = IndirectionTable.Num() * IndirectionTable.GetTypeSize();
+		const int32 MipDataBytes = IndirectionTableBytes + DistanceFieldBrickData.Num();
+
+		if (MipIndex == DistanceField::NumMips - 1)
+		{
+			VolumeDataOut.AlwaysLoadedMip.Empty(MipDataBytes);
+			VolumeDataOut.AlwaysLoadedMip.AddUninitialized(MipDataBytes);
+
+			FPlatformMemory::Memcpy(&VolumeDataOut.AlwaysLoadedMip[0], IndirectionTable.GetData(), IndirectionTableBytes);
+
+			if (DistanceFieldBrickData.Num() > 0)
+			{
+				FPlatformMemory::Memcpy(&VolumeDataOut.AlwaysLoadedMip[IndirectionTableBytes], DistanceFieldBrickData.GetData(), DistanceFieldBrickData.Num());
+			}
+		}
+		else
+		{
+			OutMip.BulkOffset = StreamableMipData.Num();
+			StreamableMipData.AddUninitialized(MipDataBytes);
+			OutMip.BulkSize = StreamableMipData.Num() - OutMip.BulkOffset;
+			checkf(OutMip.BulkSize > 0, TEXT("DynamicMeshComponent - BulkSize was 0 with %ux%ux%u indirection"), IndirectionDimensions.X, IndirectionDimensions.Y, IndirectionDimensions.Z);
+
+			FPlatformMemory::Memcpy(&StreamableMipData[OutMip.BulkOffset], IndirectionTable.GetData(), IndirectionTableBytes);
+
+			if (DistanceFieldBrickData.Num() > 0)
+			{
+				FPlatformMemory::Memcpy(&StreamableMipData[OutMip.BulkOffset + IndirectionTableBytes], DistanceFieldBrickData.GetData(), DistanceFieldBrickData.Num());
+			}
+		}
+	
+		if ( Progress.Cancelled() ) { return false; }
+
+		OutMip.IndirectionDimensions = IndirectionDimensions;
+		OutMip.DistanceFieldToVolumeScaleBias = DistanceFieldToVolumeScaleBias;
+		OutMip.NumDistanceFieldBricks = NumBricks;
+
+		// Account for the border voxels we added
+		const FVector3f VirtualUVMin = FVector3f(DistanceField::MeshDistanceFieldObjectBorder) / FVector3f(IndirectionDimensions * DistanceField::UniqueDataBrickSize);
+		const FVector3f VirtualUVSize = FVector3f(IndirectionDimensions * DistanceField::UniqueDataBrickSize - FIntVector(2 * DistanceField::MeshDistanceFieldObjectBorder)) / FVector3f(IndirectionDimensions * DistanceField::UniqueDataBrickSize);
+		
+		const FVector3f VolumePositionExtent = LocalSpaceMeshBounds.GetExtent() * LocalToVolumeScale;
+
+		// [-VolumePositionExtent, VolumePositionExtent] -> [VirtualUVMin, VirtualUVMin + VirtualUVSize]
+		OutMip.VolumeToVirtualUVScale = VirtualUVSize / (2 * VolumePositionExtent);
+		OutMip.VolumeToVirtualUVAdd = VolumePositionExtent * OutMip.VolumeToVirtualUVScale + VirtualUVMin;
+	}
+
+	VolumeDataOut.bMostlyTwoSided = bGenerateAsIfTwoSided;
+	VolumeDataOut.LocalSpaceMeshBounds = LocalSpaceMeshBounds;
+
+	if ( Progress.Cancelled() ) { return false; }
+
+	VolumeDataOut.StreamableMips.Lock(LOCK_READ_WRITE);
+	uint8* Ptr = (uint8*)VolumeDataOut.StreamableMips.Realloc(StreamableMipData.Num());
+	FMemory::Memcpy(Ptr, StreamableMipData.GetData(), StreamableMipData.Num());
+	VolumeDataOut.StreamableMips.Unlock();
+	VolumeDataOut.StreamableMips.SetBulkDataFlags(BULKDATA_Force_NOT_InlinePayload);
+
+	const float BuildTime = (float)(FPlatformTime::Seconds() - StartTime);
+		 
+	if (BuildTime > 1.0f)
+	{
+		UE_LOG(LogGeometry, Log, TEXT("DynamicMeshComponent - Finished distance field build in %.1fs - %ux%ux%u sparse distance field, %.1fMb total, %.1fMb always loaded, %u%% occupied, %u triangles"),
+			BuildTime,
+			Mip0IndirectionDimensions.X * DistanceField::UniqueDataBrickSize,
+			Mip0IndirectionDimensions.Y * DistanceField::UniqueDataBrickSize,
+			Mip0IndirectionDimensions.Z * DistanceField::UniqueDataBrickSize,
+			(VolumeDataOut.GetResourceSizeBytes() + VolumeDataOut.StreamableMips.GetBulkDataSize()) / 1024.0f / 1024.0f,
+			(VolumeDataOut.AlwaysLoadedMip.GetAllocatedSize()) / 1024.0f / 1024.0f,
+			FMath::RoundToInt(100.0f * VolumeDataOut.Mips[0].NumDistanceFieldBricks / (float)(Mip0IndirectionDimensions.X * Mip0IndirectionDimensions.Y * Mip0IndirectionDimensions.Z)),
+			Mesh.TriangleCount());
+	}
+
+	return true;
+}
+
+
+TUniquePtr<FDistanceFieldVolumeData> FBaseDynamicMeshSceneProxy::ComputeDistanceFieldForMesh(
+	const FDynamicMesh3& Mesh, 
+	FProgressCancel& Progress,
+	float DistanceFieldResolutionScale, 
+	bool bGenerateAsIfTwoSided)
+{
+	TUniquePtr<FDistanceFieldVolumeData> NewDistanceField = MakeUnique<FDistanceFieldVolumeData>();
+	bool bCompleted = DynamicMesh_GenerateSignedDistanceFieldVolumeData( Mesh, 
+		DistanceFieldResolutionScale, bGenerateAsIfTwoSided, *NewDistanceField, Progress);
+	if (bCompleted)
+	{
+		return NewDistanceField;
+	}
+	return TUniquePtr<FDistanceFieldVolumeData>();
+}
+
+void FBaseDynamicMeshSceneProxy::SetNewDistanceField(TSharedPtr<FDistanceFieldVolumeData> NewDistanceField, bool bInInitialize)
+{
+	if (DistanceField.IsValid() && NewDistanceField.IsValid() && DistanceField.Get() == NewDistanceField.Get())
+	{
+		checkSlow(false); // we don't expect this to be called when no work needs to be done
+		return;
+	}
+
+	// wait for end of frame
+	if (!bInInitialize)
+	{
+		// Note this requires us to be on the game thread
+		check(IsInGameThread());
+		FlushRenderingCommands();
+	}
+
+	DistanceField = NewDistanceField;
+	bDistanceFieldValid = DistanceField.IsValid();
+	bSupportsDistanceFieldRepresentation = bDistanceFieldValid;
+	bAffectDistanceFieldLighting = bDistanceFieldValid;
+
+	// lumen visibility may change depending on the presence of a valid distance field
+	UpdateVisibleInLumenScene();
+}
+
+
+
+
+
diff --git a/Engine/Source/Runtime/GeometryFramework/Private/Components/DynamicMeshComponent.cpp b/Engine/Source/Runtime/GeometryFramework/Private/Components/DynamicMeshComponent.cpp
--- a/Engine/Source/Runtime/GeometryFramework/Private/Components/DynamicMeshComponent.cpp
+++ b/Engine/Source/Runtime/GeometryFramework/Private/Components/DynamicMeshComponent.cpp
@@ -90,6 +90,11 @@ UDynamicMeshComponent::UDynamicMeshComponent(const FObjectInitializer& ObjectIni
 
 	MeshObjectChangedHandle = MeshObject->OnMeshChanged().AddUObject(this, &UDynamicMeshComponent::OnMeshObjectChanged);
 
+	DistanceFieldComputeQueue.OnComputeCompleted = [this](TUniquePtr<FDistanceFieldVolumeData> NewData)
+	{
+		OnNewDistanceFieldData_Async(MoveTemp(NewData));
+	};
+
 	ResetProxy();
 }
 
@@ -134,8 +139,24 @@ void UDynamicMeshComponent::PostLoad()
 
 	// make sure BodySetup is created
 	GetBodySetup();
+
+	// Note we don't serialize the distance field, so recompute on load (and below in PostEditImport, on duplicate/copy)
+	// (Note if we do switch to serializing it, it is DDC cache data and not versioned, so must be checked vs its DDC key and potentially invalidated)
+	if (CurrentDistanceField.IsValid() != (DistanceFieldMode != EDynamicMeshComponentDistanceFieldMode::NoDistanceField))
+	{
+		OnNewDistanceFieldMode();
+	}
 }
 
+void UDynamicMeshComponent::PostEditImport()
+{
+	Super::PostEditImport();
+
+	if (CurrentDistanceField.IsValid() != (DistanceFieldMode != EDynamicMeshComponentDistanceFieldMode::NoDistanceField))
+	{
+		OnNewDistanceFieldMode();
+	}
+}
 
 #if WITH_EDITOR
 void UDynamicMeshComponent::PostEditChangeProperty(FPropertyChangedEvent& PropertyChangedEvent)
@@ -972,6 +993,17 @@ FPrimitiveSceneProxy* UDynamicMeshComponent::CreateSceneProxy()
 			NewProxy->Initialize();
 		}
 
+		// set new distance field
+		if ( DistanceFieldMode != EDynamicMeshComponentDistanceFieldMode::NoDistanceField )
+		{
+			DistanceFieldLock.Lock();
+			if ( CurrentDistanceField.IsValid() )
+			{
+				NewProxy->SetNewDistanceField(CurrentDistanceField, true);
+			}
+			DistanceFieldLock.Unlock();
+		}
+
 		NewProxy->SetVerifyUsedMaterials(bProxyVerifyUsedMaterials);
 	}
 
@@ -990,6 +1022,85 @@ void UDynamicMeshComponent::NotifyMaterialSetUpdated()
 }
 
 
+
+void UDynamicMeshComponent::OnNewDistanceFieldMode()
+{
+	UpdateDistanceField();
+}
+
+
+void UDynamicMeshComponent::UpdateDistanceField()
+{
+	if (DistanceFieldMode == EDynamicMeshComponentDistanceFieldMode::NoDistanceField)
+	{
+		FScopeLock Lock(&DistanceFieldLock);
+		CurrentDistanceField = TSharedPtr<FDistanceFieldVolumeData>();
+		if (GetCurrentSceneProxy() != nullptr)
+		{
+			GetCurrentSceneProxy()->SetNewDistanceField(CurrentDistanceField, false);
+		}
+		return;
+	}
+
+	DistanceFieldComputeQueue.LaunchJob(TEXT("DynamicMeshComponentDistanceField"), 
+		[this](FProgressCancel& Progress)
+		{
+			return ComputeNewDistanceField_TaskFunction(Progress);
+		});
+}
+
+
+TUniquePtr<FDistanceFieldVolumeData> UDynamicMeshComponent::ComputeNewDistanceField_TaskFunction(FProgressCancel& Progress)
+{
+	TUniquePtr<FDistanceFieldVolumeData> NewDistanceField;
+	ProcessMesh( [&](const FDynamicMesh3& ReadMesh)
+	{
+		float DistanceFieldResolutionScale = 1.0f;
+		bool bMostlyTwoSided = false;
+		if (ReadMesh.Attributes() && ReadMesh.Attributes()->GetMaterialID())
+		{
+			TArray<bool> MatIsTwoSided;
+			MatIsTwoSided.SetNumUninitialized(BaseMaterials.Num());
+			for (int32 Idx = 0; Idx < BaseMaterials.Num(); ++Idx)
+			{
+				MatIsTwoSided[Idx] = BaseMaterials[Idx] ? BaseMaterials[Idx]->IsTwoSided() : false;
+			}
+			const FDynamicMeshMaterialAttribute* Materials = ReadMesh.Attributes()->GetMaterialID();
+			int32 TwoSidedTriCount = 0;
+			for (int32 TID : ReadMesh.TriangleIndicesItr())
+			{
+				int32 MID = Materials->GetValue(TID);
+				TwoSidedTriCount += MatIsTwoSided.IsValidIndex(MID) ? (int32)MatIsTwoSided[MID] : 0;
+			}
+			bMostlyTwoSided = TwoSidedTriCount * 2 >= ReadMesh.TriangleCount();
+		}
+		// TODO: Consider optionally building distance field via collision mesh, allowing this mesh to be unlocked. (But note the collision mesh may not be available, especially for a large mesh)
+		NewDistanceField =
+			FDynamicMeshSceneProxy::ComputeDistanceFieldForMesh(ReadMesh, Progress, DistanceFieldResolutionScale, bMostlyTwoSided);
+	});
+	return NewDistanceField;
+}
+
+void UDynamicMeshComponent::OnNewDistanceFieldData_Async(TUniquePtr<FDistanceFieldVolumeData> NewData)
+{
+	// WARNING: this function will be called from TAsyncComponentDataComputeQueue background tasks
+
+	TSharedPtr<FDistanceFieldVolumeData> NewDistanceField(NewData.Release());
+
+	DistanceFieldLock.Lock();
+	CurrentDistanceField = NewDistanceField;
+	if (GetCurrentSceneProxy() != nullptr)
+	{
+		// the new distance field will be set when the scene proxy is re-created
+		MarkRenderStateDirty();
+	}
+	DistanceFieldLock.Unlock();
+}
+
+
+
+
+
 void UDynamicMeshComponent::SetTriangleColorFunction(
 	TUniqueFunction<FColor(const FDynamicMesh3*, int)> TriangleColorFuncIn,
 	EDynamicMeshComponentRenderUpdateMode UpdateMode)
@@ -1247,6 +1358,8 @@ void UDynamicMeshComponent::InternalOnMeshUpdated()
 	{
 		RebuildPhysicsData();
 	}
+
+	UpdateDistanceField();
 }
 
 bool UDynamicMeshComponent::GetTriMeshSizeEstimates(struct FTriMeshCollisionDataEstimates& OutTriMeshEstimates, bool bInUseAllTriData) const
diff --git a/Engine/Source/Runtime/GeometryFramework/Private/Components/DynamicMeshSceneProxy.h b/Engine/Source/Runtime/GeometryFramework/Private/Components/DynamicMeshSceneProxy.h
--- a/Engine/Source/Runtime/GeometryFramework/Private/Components/DynamicMeshSceneProxy.h
+++ b/Engine/Source/Runtime/GeometryFramework/Private/Components/DynamicMeshSceneProxy.h
@@ -85,6 +85,10 @@ public:
 			bIsSingleBuffer = true;
 			InitializeSingleBufferSet(RenderBufferSets[0]);
 		}
+
+		// set (or clear) simple Lumen cards via the bounding box
+		// TODO: Implement a better method to set up the lumen cards
+		UpdateLumenCardsFromBounds();
 	}
 
 
diff --git a/Engine/Source/Runtime/GeometryFramework/Public/Components/BaseDynamicMeshComponent.h b/Engine/Source/Runtime/GeometryFramework/Public/Components/BaseDynamicMeshComponent.h
--- a/Engine/Source/Runtime/GeometryFramework/Public/Components/BaseDynamicMeshComponent.h
+++ b/Engine/Source/Runtime/GeometryFramework/Public/Components/BaseDynamicMeshComponent.h
@@ -87,6 +87,18 @@ enum class EDynamicMeshDrawPath : uint8
 };
 
 
+/**
+ * Mesh Signed Distance Field (SDF) mode
+ */
+UENUM(BlueprintType)
+enum class EDynamicMeshComponentDistanceFieldMode : uint8
+{
+	/** Do not compute a distance field */
+	NoDistanceField = 0,
+	/** Compute a distance field in a background thread */
+	AsyncCPUDistanceField = 1 UMETA(DisplayName = "Async CPU Distance Field")
+};
+
 /**
  * Color Transform to apply to Vertex Colors when converting from internal DynamicMesh
  * Color attributes (eg Color Overlay stored in FVector4f) to RHI Render Buffers (FColor).
@@ -97,11 +109,11 @@ UENUM(BlueprintType)
 enum class EDynamicMeshVertexColorTransformMode : uint8
 {
 	/** Do not apply any color-space transform to Vertex Colors */
-	NoTransform,
+	NoTransform = 0,
 	/** Assume Vertex Colors are in Linear space and transform to SRGB */
 	LinearToSRGB,
 	/** Assume Vertex Colors are in SRGB space and convert to Linear */
-	SRGBToLinear
+	SRGBToLinear UMETA(DisplayName = "SRGB To Linear")
 };
 
 
@@ -343,7 +355,7 @@ public:
 	GEOMETRYFRAMEWORK_API virtual void SetEnableFlatShading(bool bEnable);
 
 	/**
-	 * @return active Color used for Constant Color Override Mode
+	 * @return Whether mesh will be rendered with per-triangle normals
 	 */
 	UFUNCTION(BlueprintCallable, Category = "Dynamic Mesh Component|Rendering")
 	virtual bool GetFlatShadingEnabled() const { return bEnableFlatShading; }
@@ -542,6 +554,32 @@ public:
 	UFUNCTION(BlueprintPure, Category = "Dynamic Mesh Component")
 	GEOMETRYFRAMEWORK_API virtual EDynamicMeshDrawPath GetMeshDrawPath() const;
 
+protected:
+	/** Controls how distance field is computed */
+	UPROPERTY(EditAnywhere, BlueprintReadWrite, 
+				BlueprintSetter = SetDistanceFieldMode, BlueprintGetter = GetDistanceFieldMode,
+					Category = "Dynamic Mesh Component|Rendering")
+	EDynamicMeshComponentDistanceFieldMode DistanceFieldMode = EDynamicMeshComponentDistanceFieldMode::NoDistanceField;
+
+public:
+	/**
+	 * Configure Distance Field computation mode
+	 */
+	UFUNCTION(BlueprintCallable, Category = "Dynamic Mesh Component")
+	GEOMETRYFRAMEWORK_API virtual void SetDistanceFieldMode(EDynamicMeshComponentDistanceFieldMode NewDistFieldMode);
+
+	/**
+	 * @return Method used to compute the distance field
+	 */
+	UFUNCTION(BlueprintPure, Category = "Dynamic Mesh Component")
+	GEOMETRYFRAMEWORK_API virtual EDynamicMeshComponentDistanceFieldMode GetDistanceFieldMode() const;
+
+protected:
+	// this will be called if a change to distance field mode is detected
+	GEOMETRYFRAMEWORK_API virtual void OnNewDistanceFieldMode();
+
+
+
 	//===============================================================================================================
 	// Standard Component interfaces
 	//
diff --git a/Engine/Source/Runtime/GeometryFramework/Public/Components/BaseDynamicMeshSceneProxy.h b/Engine/Source/Runtime/GeometryFramework/Public/Components/BaseDynamicMeshSceneProxy.h
--- a/Engine/Source/Runtime/GeometryFramework/Public/Components/BaseDynamicMeshSceneProxy.h
+++ b/Engine/Source/Runtime/GeometryFramework/Public/Components/BaseDynamicMeshSceneProxy.h
@@ -13,6 +13,9 @@
 #include "DynamicMeshBuilder.h"
 #include "Components/BaseDynamicMeshComponent.h"
 #include "RayTracingGeometry.h"
+#include "Templates/PimplPtr.h"
+#include "Util/ProgressCancel.h"
+#include "DistanceFieldAtlas.h"
 
 #include "PhysicsEngine/AggregateGeom.h"
 
@@ -27,6 +30,7 @@ class FDynamicPrimitiveUniformBuffer;
 class FMaterialRenderProxy;
 class UMaterialInterface;
 struct FRayTracingMaterialGatheringContext;
+class FCardRepresentationData;
 
 /**
  * FMeshRenderBufferSet stores a set of RenderBuffers for a mesh
@@ -1018,6 +1022,42 @@ public:
 
 #endif // RHI_RAYTRACING
 
+
+	//
+	// Lumen APIs
+	//
+
+
+public:
+	GEOMETRYFRAMEWORK_API virtual const FCardRepresentationData* GetMeshCardRepresentation() const override;
+
+protected:
+	TPimplPtr<FCardRepresentationData> MeshCards;
+	bool bMeshCardsValid = false;
+
+	// Helper to set lumen cards
+	void UpdateLumenCardsFromBounds();
+
+
+public:
+	GEOMETRYFRAMEWORK_API virtual void GetDistanceFieldAtlasData(const class FDistanceFieldVolumeData*& OutDistanceFieldData, float& SelfShadowBias) const;
+	GEOMETRYFRAMEWORK_API virtual void GetDistanceFieldInstanceData(TArray<FRenderTransform>& InstanceLocalToPrimitiveTransforms) const override;
+	GEOMETRYFRAMEWORK_API virtual bool HasDistanceFieldRepresentation() const override;
+	GEOMETRYFRAMEWORK_API virtual bool HasDynamicIndirectShadowCasterRepresentation() const override;
+
+
+protected:
+	TSharedPtr<FDistanceFieldVolumeData> DistanceField;
+	bool bDistanceFieldValid = false;
+public:
+	static TUniquePtr<FDistanceFieldVolumeData> ComputeDistanceFieldForMesh(
+		const FDynamicMesh3& Mesh, 
+		FProgressCancel& Progress,
+		float DistanceFieldResolutionScale = 1.0, 
+		bool bGenerateAsIfTwoSided = false );
+
+	GEOMETRYFRAMEWORK_API void SetNewDistanceField(TSharedPtr<FDistanceFieldVolumeData> NewDistanceField, bool bInInitialize);
+
 public:
 	// Set the collision data to use for debug drawing, or do nothing if debug drawing is not enabled
 	GEOMETRYFRAMEWORK_API void SetCollisionData();
diff --git a/Engine/Source/Runtime/GeometryFramework/Public/Components/DynamicMeshComponent.h b/Engine/Source/Runtime/GeometryFramework/Public/Components/DynamicMeshComponent.h
--- a/Engine/Source/Runtime/GeometryFramework/Public/Components/DynamicMeshComponent.h
+++ b/Engine/Source/Runtime/GeometryFramework/Public/Components/DynamicMeshComponent.h
@@ -12,6 +12,11 @@
 #include "UDynamicMesh.h"
 #include "PhysicsEngine/BodySetup.h"
 
+#include "Misc/ScopeLock.h"
+#include "Tasks/Task.h"
+#include "Util/ProgressCancel.h"
+#include "DistanceFieldAtlas.h"
+
 #include "DynamicMeshComponent.generated.h"
 
 // predecl
@@ -48,6 +53,108 @@ enum class EDynamicMeshComponentRenderUpdateMode
 };
 
 
+template<typename DataType>
+struct TAsyncComponentDataComputeQueue
+{
+	std::atomic<int> JobCounter = 0;
+	bool bIsShuttingDown = false;
+
+	~TAsyncComponentDataComputeQueue()
+	{
+		WaitForAllJobsDuringShutdown();
+	}
+
+	struct FComputeJob
+	{
+		UE::Tasks::FTask Task;
+		int JobTimestamp = 0;
+		TUniquePtr<FProgressCancel> Progress;
+		bool bCancelled = false;
+		bool bHasCompleted = false;
+	};
+
+	TArray<TUniquePtr<FComputeJob>> PendingJobs;
+	FCriticalSection PendingJobsLock;
+
+	// This function will be called w/ the data computed by a job
+	// when it finishes, if it is still valid
+	// Note this method will be run from a background thread, and should be thread-safe!
+	TFunction<void(TUniquePtr<DataType> NewData)> OnComputeCompleted;
+
+	void LaunchJob(const TCHAR* DebugName, TFunction<TUniquePtr<DataType>(FProgressCancel& Progress)> JobWork)
+	{
+		// OnComputeCompleted function must be set to something
+		check(OnComputeCompleted);
+
+		if (!ensure(bIsShuttingDown == false))
+		{
+			return;
+		}
+
+		JobCounter++;
+		int CurrentTimestamp = JobCounter;
+
+		// cancel any existing jobs and clear them out if they have returned
+		{
+			FScopeLock RemovePending(&PendingJobsLock);
+			for (int32 k = 0; k < PendingJobs.Num(); ++k)
+			{
+				PendingJobs[k]->bCancelled = true;
+				if (PendingJobs[k]->bHasCompleted)
+				{
+					PendingJobs.RemoveAtSwap(k, 1, EAllowShrinking::No);
+					k--;		// reconsider element that was just swapped in to this position
+				}
+			}
+		}
+
+		// set up the new job
+		TUniquePtr<FComputeJob> NewJob = MakeUnique<FComputeJob>();
+		FComputeJob* JobPtr = NewJob.Get();
+		NewJob->Progress = MakeUnique<FProgressCancel>();
+		FProgressCancel* ProgressPtr = NewJob->Progress.Get();
+		NewJob->Progress->CancelF = [this, JobPtr]() { return bIsShuttingDown || JobPtr->bCancelled; };
+		NewJob->JobTimestamp = CurrentTimestamp;
+
+		// launch it
+		NewJob->Task = UE::Tasks::Launch(DebugName, 
+			[this, JobWork, JobPtr]() {
+				// TODO: limit the number of in-progress active jobs
+				TUniquePtr<DataType> Result = JobWork( *JobPtr->Progress );
+				if (JobPtr->JobTimestamp == this->JobCounter && JobPtr->bCancelled == false)
+				{
+					if (!bIsShuttingDown && this->OnComputeCompleted)
+					{
+						this->OnComputeCompleted(MoveTemp(Result));
+					}
+				}
+				JobPtr->bHasCompleted = true;
+			},
+			LowLevelTasks::ETaskPriority::BackgroundNormal, 
+			UE::Tasks::EExtendedTaskPriority::None );
+
+		// add new job
+		{
+			FScopeLock AddJob(&PendingJobsLock);
+			PendingJobs.Add(MoveTemp(NewJob));
+		}
+	}
+
+
+	void WaitForAllJobsDuringShutdown()
+	{
+		bIsShuttingDown = true;
+		FScopeLock PendingLock(&PendingJobsLock);
+		for (int32 k = 0; k < PendingJobs.Num(); ++k)
+		{
+			UE::Tasks::Wait({PendingJobs[k]->Task});
+		}
+	}
+
+};
+
+
+
 /** 
  * UDynamicMeshComponent is a mesh component similar to UProceduralMeshComponent,
  * except it bases the renderable geometry off an internal UDynamicMesh instance (which
@@ -534,6 +641,23 @@ protected:
 	GEOMETRYFRAMEWORK_API void UpdateAutoCalculatedTangents();
 
 
+	//===============================================================================================================
+	//
+	// Distance Field Support
+	//
+protected:
+	FCriticalSection DistanceFieldLock;
+	TSharedPtr<FDistanceFieldVolumeData> CurrentDistanceField;
+
+	GEOMETRYFRAMEWORK_API virtual void UpdateDistanceField();
+
+	TAsyncComponentDataComputeQueue<FDistanceFieldVolumeData> DistanceFieldComputeQueue;
+	GEOMETRYFRAMEWORK_API virtual TUniquePtr<FDistanceFieldVolumeData> ComputeNewDistanceField_TaskFunction(FProgressCancel& Progress);
+	GEOMETRYFRAMEWORK_API virtual void OnNewDistanceFieldData_Async(TUniquePtr<FDistanceFieldVolumeData> NewData);
+
+	// UBaseDynamicMeshComponent API
+	GEOMETRYFRAMEWORK_API virtual void OnNewDistanceFieldMode() override;
+
 
 	//===============================================================================================================
 	//
@@ -716,7 +840,8 @@ protected:
 
 	//~ UObject Interface.
 	GEOMETRYFRAMEWORK_API virtual void Serialize(FArchive& Ar) override;
-	GEOMETRYFRAMEWORK_API virtual void PostLoad() override;
+	GEOMETRYFRAMEWORK_API virtual void PostLoad() override; // called after load
+	GEOMETRYFRAMEWORK_API virtual void PostEditImport() override; // called after duplicate/copy
 	GEOMETRYFRAMEWORK_API virtual void BeginDestroy() override;
 #if WITH_EDITOR
 	GEOMETRYFRAMEWORK_API void PostEditChangeProperty(FPropertyChangedEvent& PropertyChangedEvent) override;
