add a limit to the number of simultaneous dynamic mesh component distance field builds that can be active
+ make the background thread operate on a (geo-only) copy of the mesh (because ProcessMesh doesn't seem to actually lock access to the mesh, and can't really since code can easily grab the mesh pointer via GetMesh() also)
+ make the post-compute MarkRenderStateDirty() call happen on the game thread (on tick), so it won't be called at an unsafe time

diff --git a/Engine/Source/Runtime/GeometryFramework/Private/Components/DynamicMeshComponent.cpp b/Engine/Source/Runtime/GeometryFramework/Private/Components/DynamicMeshComponent.cpp
--- a/Engine/Source/Runtime/GeometryFramework/Private/Components/DynamicMeshComponent.cpp
+++ b/Engine/Source/Runtime/GeometryFramework/Private/Components/DynamicMeshComponent.cpp
@@ -7,6 +7,7 @@
 #include "Engine/World.h"
 #include "Materials/Material.h"
 #include "Async/Async.h"
+#include "HAL/UESemaphore.h"
 #include "Engine/CollisionProfile.h"
 #include "PhysicsEngine/PhysicsSettings.h"
 
@@ -1042,21 +1043,12 @@ void UDynamicMeshComponent::UpdateDistanceField()
 		return;
 	}
 
-	DistanceFieldComputeQueue.LaunchJob(TEXT("DynamicMeshComponentDistanceField"), 
-		[this](FProgressCancel& Progress)
-		{
-			return ComputeNewDistanceField_TaskFunction(Progress);
-		});
-}
-
-
-TUniquePtr<FDistanceFieldVolumeData> UDynamicMeshComponent::ComputeNewDistanceField_TaskFunction(FProgressCancel& Progress)
-{
-	TUniquePtr<FDistanceFieldVolumeData> NewDistanceField;
-	ProcessMesh( [&](const FDynamicMesh3& ReadMesh)
+	// For safety, run the distance field compute on a (geometry-only) copy of the mesh
+	FDynamicMesh3 GeoOnlyCopy;
+	// Compute whether the mesh uses mainly two-sided materials before, as this is the only info the distance field compute needs from the mesh attributes
+	bool bMostlyTwoSided = false;
+	ProcessMesh([&](const FDynamicMesh3& ReadMesh)
 	{
-		float DistanceFieldResolutionScale = 1.0f;
-		bool bMostlyTwoSided = false;
 		if (ReadMesh.Attributes() && ReadMesh.Attributes()->GetMaterialID())
 		{
 			TArray<bool> MatIsTwoSided;
@@ -1074,10 +1066,35 @@ TUniquePtr<FDistanceFieldVolumeData> UDynamicMeshComponent::ComputeNewDistanceFi
 			}
 			bMostlyTwoSided = TwoSidedTriCount * 2 >= ReadMesh.TriangleCount();
 		}
-		// TODO: Consider optionally building distance field via collision mesh, allowing this mesh to be unlocked. (But note the collision mesh may not be available, especially for a large mesh)
-		NewDistanceField =
-			FDynamicMeshSceneProxy::ComputeDistanceFieldForMesh(ReadMesh, Progress, DistanceFieldResolutionScale, bMostlyTwoSided);
+
+		GeoOnlyCopy.Copy(ReadMesh, false, false, false, false);
 	});
+	DistanceFieldComputeQueue.LaunchJob(TEXT("DynamicMeshComponentDistanceField"), 
+		[this, MovedGeoOnlyCopy = MoveTemp(GeoOnlyCopy), bMostlyTwoSided](FProgressCancel& Progress)
+		{
+			return ComputeNewDistanceField_TaskFunction(Progress, MovedGeoOnlyCopy, bMostlyTwoSided);
+		});
+}
+
+
+TUniquePtr<FDistanceFieldVolumeData> UDynamicMeshComponent::ComputeNewDistanceField_TaskFunction(FProgressCancel& Progress, const FDynamicMesh3& Mesh, bool bMostlyTwoSided)
+{
+	// todo: consider making the number of concurrent distance field computes configurable
+	constexpr int32 MaxConcurrentComputes = 3;
+	static FSemaphore ComputesCountSemaphore(MaxConcurrentComputes, MaxConcurrentComputes);
+	
+	ComputesCountSemaphore.Acquire();
+	if (Progress.Cancelled())
+	{
+		ComputesCountSemaphore.Release();
+		return nullptr;
+	}
+
+	TUniquePtr<FDistanceFieldVolumeData> NewDistanceField;
+	float DistanceFieldResolutionScale = 1.0f;
+	NewDistanceField =
+		FDynamicMeshSceneProxy::ComputeDistanceFieldForMesh(Mesh, Progress, DistanceFieldResolutionScale, bMostlyTwoSided);
+	ComputesCountSemaphore.Release();
 	return NewDistanceField;
 }
 
@@ -1091,8 +1108,15 @@ void UDynamicMeshComponent::OnNewDistanceFieldData_Async(TUniquePtr<FDistanceFie
 	CurrentDistanceField = NewDistanceField;
 	if (GetCurrentSceneProxy() != nullptr)
 	{
-		// the new distance field will be set when the scene proxy is re-created
-		MarkRenderStateDirty();
+		// mark render state dirty on the game thread to ensure it updates at a safe time (e.g., cannot update when bPostTickComponentUpdate == true)
+		AsyncTask(
+			ENamedThreads::GameThread,
+			[this]()
+			{
+				// the new distance field will be set when the scene proxy is re-created
+				MarkRenderStateDirty();
+			}
+		);
 	}
 	DistanceFieldLock.Unlock();
 }
diff --git a/Engine/Source/Runtime/GeometryFramework/Public/Components/DynamicMeshComponent.h b/Engine/Source/Runtime/GeometryFramework/Public/Components/DynamicMeshComponent.h
--- a/Engine/Source/Runtime/GeometryFramework/Public/Components/DynamicMeshComponent.h
+++ b/Engine/Source/Runtime/GeometryFramework/Public/Components/DynamicMeshComponent.h
@@ -652,12 +652,15 @@ protected:
 	GEOMETRYFRAMEWORK_API virtual void UpdateDistanceField();
 
 	TAsyncComponentDataComputeQueue<FDistanceFieldVolumeData> DistanceFieldComputeQueue;
-	GEOMETRYFRAMEWORK_API virtual TUniquePtr<FDistanceFieldVolumeData> ComputeNewDistanceField_TaskFunction(FProgressCancel& Progress);
 	GEOMETRYFRAMEWORK_API virtual void OnNewDistanceFieldData_Async(TUniquePtr<FDistanceFieldVolumeData> NewData);
 
 	// UBaseDynamicMeshComponent API
 	GEOMETRYFRAMEWORK_API virtual void OnNewDistanceFieldMode() override;
 
+private:
+	// Internal method to compute the distance field, run in a background thread.
+	TUniquePtr<FDistanceFieldVolumeData> ComputeNewDistanceField_TaskFunction(FProgressCancel& Progress, const FDynamicMesh3& Mesh, bool bMostlyTwoSided);
+
 
 	//===============================================================================================================
 	//
