Expose mehtods that create and update FMeshRenderBufferSets from FDynamicMesh3 so this can be used outside the DynamicMeshSceneProxy.

diff --git a/Engine/Plugins/Runtime/MeshModelingToolset/Source/ModelingComponents/Private/Components/OctreeDynamicMeshComponent.cpp b/Engine/Plugins/Runtime/MeshModelingToolset/Source/ModelingComponents/Private/Components/OctreeDynamicMeshComponent.cpp
--- a/Engine/Plugins/Runtime/MeshModelingToolset/Source/ModelingComponents/Private/Components/OctreeDynamicMeshComponent.cpp
+++ b/Engine/Plugins/Runtime/MeshModelingToolset/Source/ModelingComponents/Private/Components/OctreeDynamicMeshComponent.cpp
@@ -239,8 +239,8 @@ FPrimitiveSceneProxy* UOctreeDynamicMeshComponent::CreateSceneProxy()
 
 		if (TriangleColorFunc != nullptr)
 		{
-			NewProxy->bUsePerTriangleColor = true;
-			NewProxy->PerTriangleColorFunc = [this](const FDynamicMesh3* MeshIn, int TriangleID) { return GetTriangleColor(TriangleID); };
+			NewProxy->MeshRenderBufferSetConverter.bUsePerTriangleColor = true;
+			NewProxy->MeshRenderBufferSetConverter.PerTriangleColorFunc = [this](const FDynamicMesh3* MeshIn, int TriangleID) { return GetTriangleColor(TriangleID); };
 		}
 
 		OctreeCut = MakeUnique<FDynamicMeshOctree3::FTreeCutSet>();
diff --git a/Engine/Source/Runtime/GeometryFramework/Private/Components/BaseDynamicMeshSceneProxy.cpp b/Engine/Source/Runtime/GeometryFramework/Private/Components/BaseDynamicMeshSceneProxy.cpp
--- a/Engine/Source/Runtime/GeometryFramework/Private/Components/BaseDynamicMeshSceneProxy.cpp
+++ b/Engine/Source/Runtime/GeometryFramework/Private/Components/BaseDynamicMeshSceneProxy.cpp
@@ -51,18 +51,19 @@ namespace UE::DynamicMesh
 FBaseDynamicMeshSceneProxy::FBaseDynamicMeshSceneProxy(UBaseDynamicMeshComponent* Component)
 	: FPrimitiveSceneProxy(Component),
 	ParentBaseComponent(Component),
-	ColorSpaceTransformMode(Component->GetVertexColorSpaceTransformMode()),
 	bEnableRaytracing(Component->GetEnableRaytracing()),
 	bEnableViewModeOverrides(Component->GetViewModeOverridesEnabled()),
 	bPreferStaticDrawPath(Component->GetMeshDrawPath() == EDynamicMeshDrawPath::StaticDraw)
 {
+	MeshRenderBufferSetConverter.ColorSpaceTransformMode = Component->GetVertexColorSpaceTransformMode();
+
 	if (Component->GetColorOverrideMode() == EDynamicMeshComponentColorOverrideMode::Constant)
 	{
-		ConstantVertexColor = Component->GetConstantOverrideColor();
-		bIgnoreVertexColors = true;
+		MeshRenderBufferSetConverter.ConstantVertexColor = Component->GetConstantOverrideColor();
+		MeshRenderBufferSetConverter.bIgnoreVertexColors = true;
 	}
 
-	bUsePerTriangleNormals = Component->GetFlatShadingEnabled();
+	MeshRenderBufferSetConverter.bUsePerTriangleNormals = Component->GetFlatShadingEnabled();
 	
 	SetCollisionData();
 
@@ -327,7 +328,7 @@ void FBaseDynamicMeshSceneProxy::GetDynamicMeshElements(const TArray<const FScen
 
 				// do we need separate one of these for each MeshRenderBufferSet?
 				FDynamicPrimitiveUniformBuffer& DynamicPrimitiveUniformBuffer = Collector.AllocateOneFrameResource<FDynamicPrimitiveUniformBuffer>();
-				DynamicPrimitiveUniformBuffer.Set(Collector.GetRHICommandList(), GetLocalToWorld(), PreviousLocalToWorld, GetBounds(), GetLocalBounds(), GetLocalBounds(), true, bHasPrecomputedVolumetricLightmap, bOutputVelocity, GetCustomPrimitiveData());
+				DynamicPrimitiveUniformBuffer.Set(Collector.GetRHICommandList(), GetLocalToWorld(), PreviousLocalToWorld, GetActorPosition(), GetBounds(), GetLocalBounds(), GetLocalBounds(), true, bHasPrecomputedVolumetricLightmap, bOutputVelocity, GetCustomPrimitiveData());
 
 				// If we want Wireframe-on-Shaded, we have to draw the solid. If View Mode Overrides are enabled, the solid
 				// will be replaced with it's wireframe, so we might as well not. 
diff --git a/Engine/Source/Runtime/GeometryFramework/Private/Components/DynamicMeshComponent.cpp b/Engine/Source/Runtime/GeometryFramework/Private/Components/DynamicMeshComponent.cpp
--- a/Engine/Source/Runtime/GeometryFramework/Private/Components/DynamicMeshComponent.cpp
+++ b/Engine/Source/Runtime/GeometryFramework/Private/Components/DynamicMeshComponent.cpp
@@ -440,26 +440,26 @@ void UDynamicMeshComponent::FastNotifyColorsUpdated()
 	FDynamicMeshSceneProxy* Proxy = GetCurrentSceneProxy();
 	if (Proxy && AllowFastUpdate())
 	{
-		if (HasTriangleColorFunction() && Proxy->bUsePerTriangleColor == false )
+		if (HasTriangleColorFunction() && Proxy->MeshRenderBufferSetConverter.bUsePerTriangleColor == false )
 		{
-			Proxy->bUsePerTriangleColor = true;
-			Proxy->PerTriangleColorFunc = [this](const FDynamicMesh3* MeshIn, int TriangleID) { return GetTriangleColor(MeshIn, TriangleID); };
+			Proxy->MeshRenderBufferSetConverter.bUsePerTriangleColor = true;
+			Proxy->MeshRenderBufferSetConverter.PerTriangleColorFunc = [this](const FDynamicMesh3* MeshIn, int TriangleID) { return GetTriangleColor(MeshIn, TriangleID); };
 		} 
-		else if ( !HasTriangleColorFunction() && Proxy->bUsePerTriangleColor == true)
+		else if ( !HasTriangleColorFunction() && Proxy->MeshRenderBufferSetConverter.bUsePerTriangleColor == true)
 		{
-			Proxy->bUsePerTriangleColor = false;
-			Proxy->PerTriangleColorFunc = nullptr;
+			Proxy->MeshRenderBufferSetConverter.bUsePerTriangleColor = false;
+			Proxy->MeshRenderBufferSetConverter.PerTriangleColorFunc = nullptr;
 		}
 
-		if (HasVertexColorRemappingFunction() && Proxy->bApplyVertexColorRemapping == false)
+		if (HasVertexColorRemappingFunction() && Proxy->MeshRenderBufferSetConverter.bApplyVertexColorRemapping == false)
 		{
-			Proxy->bApplyVertexColorRemapping = true;
-			Proxy->VertexColorRemappingFunc = [this](FVector4f& Color) { RemapVertexColor(Color); };
+			Proxy->MeshRenderBufferSetConverter.bApplyVertexColorRemapping = true;
+			Proxy->MeshRenderBufferSetConverter.VertexColorRemappingFunc = [this](FVector4f& Color) { RemapVertexColor(Color); };
 		}
-		else if (!HasVertexColorRemappingFunction() && Proxy->bApplyVertexColorRemapping == true)
+		else if (!HasVertexColorRemappingFunction() && Proxy->MeshRenderBufferSetConverter.bApplyVertexColorRemapping == true)
 		{
-			Proxy->bApplyVertexColorRemapping = false;
-			Proxy->VertexColorRemappingFunc = nullptr;
+			Proxy->MeshRenderBufferSetConverter.bApplyVertexColorRemapping = false;
+			Proxy->MeshRenderBufferSetConverter.VertexColorRemappingFunc = nullptr;
 		}
 
 		Proxy->FastUpdateVertices(false, false, true, false);
@@ -961,25 +961,25 @@ FPrimitiveSceneProxy* UDynamicMeshComponent::CreateSceneProxy()
 
 		if (TriangleColorFunc)
 		{
-			NewProxy->bUsePerTriangleColor = true;
-			NewProxy->PerTriangleColorFunc = [this](const FDynamicMesh3* MeshIn, int TriangleID) { return GetTriangleColor(MeshIn, TriangleID); };
+			NewProxy->MeshRenderBufferSetConverter.bUsePerTriangleColor = true;
+			NewProxy->MeshRenderBufferSetConverter.PerTriangleColorFunc = [this](const FDynamicMesh3* MeshIn, int TriangleID) { return GetTriangleColor(MeshIn, TriangleID); };
 		}
 		else if ( GetColorOverrideMode() == EDynamicMeshComponentColorOverrideMode::Polygroups )
 		{
-			NewProxy->bUsePerTriangleColor = true;
-			NewProxy->PerTriangleColorFunc = [this](const FDynamicMesh3* MeshIn, int TriangleID) { return GetGroupColor(MeshIn, TriangleID); };
+			NewProxy->MeshRenderBufferSetConverter.bUsePerTriangleColor = true;
+			NewProxy->MeshRenderBufferSetConverter.PerTriangleColorFunc = [this](const FDynamicMesh3* MeshIn, int TriangleID) { return GetGroupColor(MeshIn, TriangleID); };
 		}
 
 		if (HasVertexColorRemappingFunction())
 		{
-			NewProxy->bApplyVertexColorRemapping = true;
-			NewProxy->VertexColorRemappingFunc = [this](FVector4f& Color) { RemapVertexColor(Color); };
+			NewProxy->MeshRenderBufferSetConverter.bApplyVertexColorRemapping = true;
+			NewProxy->MeshRenderBufferSetConverter.VertexColorRemappingFunc = [this](FVector4f& Color) { RemapVertexColor(Color); };
 		}
 
 		if (SecondaryTriFilterFunc)
 		{
-			NewProxy->bUseSecondaryTriBuffers = true;
-			NewProxy->SecondaryTriFilterFunc = [this](const FDynamicMesh3* MeshIn, int32 TriangleID) 
+			NewProxy->MeshRenderBufferSetConverter.bUseSecondaryTriBuffers = true;
+			NewProxy->MeshRenderBufferSetConverter.SecondaryTriFilterFunc = [this](const FDynamicMesh3* MeshIn, int32 TriangleID)
 			{ 
 				return (SecondaryTriFilterFunc) ? SecondaryTriFilterFunc(MeshIn, TriangleID) : false;
 			};
diff --git a/Engine/Source/Runtime/GeometryFramework/Private/Components/MeshRenderBufferSet.cpp b/Engine/Source/Runtime/GeometryFramework/Private/Components/MeshRenderBufferSet.cpp
new file mode 100644
--- /dev/null
+++ b/Engine/Source/Runtime/GeometryFramework/Private/Components/MeshRenderBufferSet.cpp
@@ -0,0 +1,291 @@
+// Copyright Epic Games, Inc. All Rights Reserved. 
+
+#include "Components/MeshRenderBufferSet.h"
+
+
+/// ---------------  FMeshRenderBuffferSet ------------------ /// 
+
+void FMeshRenderBufferSet::Upload()
+{
+	FRHICommandListBase& RHICmdList = FRHICommandListImmediate::Get();
+
+	if (TriangleCount == 0)
+	{
+		return;
+	}
+
+	InitOrUpdateResource(RHICmdList, &this->PositionVertexBuffer);
+	InitOrUpdateResource(RHICmdList, &this->StaticMeshVertexBuffer);
+	InitOrUpdateResource(RHICmdList, &this->ColorVertexBuffer);
+
+	FLocalVertexFactory::FDataType Data;
+	this->PositionVertexBuffer.BindPositionVertexBuffer(&this->VertexFactory, Data);
+	this->StaticMeshVertexBuffer.BindTangentVertexBuffer(&this->VertexFactory, Data);
+	this->StaticMeshVertexBuffer.BindPackedTexCoordVertexBuffer(&this->VertexFactory, Data);
+	// currently no lightmaps support
+	//this->StaticMeshVertexBuffer.BindLightMapVertexBuffer(&this->VertexFactory, Data, LightMapIndex);
+	this->ColorVertexBuffer.BindColorVertexBuffer(&this->VertexFactory, Data);
+	this->VertexFactory.SetData(RHICmdList, Data);
+
+	InitOrUpdateResource(RHICmdList, &this->VertexFactory);
+	PositionVertexBuffer.InitResource(RHICmdList);
+	StaticMeshVertexBuffer.InitResource(RHICmdList);
+	ColorVertexBuffer.InitResource(RHICmdList);
+	VertexFactory.InitResource(RHICmdList);
+
+	if (IndexBuffer.Indices.Num() > 0)
+	{
+		IndexBuffer.InitResource(RHICmdList);
+	}
+	if (bEnableSecondaryIndexBuffer && SecondaryIndexBuffer.Indices.Num() > 0)
+	{
+		SecondaryIndexBuffer.InitResource(RHICmdList);
+	}
+
+	InvalidateRayTracingData();
+	ValidateRayTracingData();		// currently we are immediately validating. This may be revisited in future.
+}
+
+
+void FMeshRenderBufferSet::UploadVertexUpdate(bool bPositions, bool bMeshAttribs, bool bColors)
+{
+	// todo: look at calls to this function, it seems possible that TransferVertexUpdateToGPU
+	// could be used instead (which should be somewhat more efficient?). It's not clear if there
+	// are any situations where we would change vertex buffer size w/o also updating the index
+	// buffers (in which case we are fully rebuilding the buffers...)
+
+	if (TriangleCount == 0)
+	{
+		return;
+	}
+
+	FRHICommandListBase& RHICmdList = FRHICommandListImmediate::Get();
+
+	if (bPositions)
+	{
+		InitOrUpdateResource(RHICmdList, &this->PositionVertexBuffer);
+	}
+	if (bMeshAttribs)
+	{
+		InitOrUpdateResource(RHICmdList, &this->StaticMeshVertexBuffer);
+	}
+	if (bColors)
+	{
+		InitOrUpdateResource(RHICmdList, &this->ColorVertexBuffer);
+	}
+
+	FLocalVertexFactory::FDataType Data;
+	this->PositionVertexBuffer.BindPositionVertexBuffer(&this->VertexFactory, Data);
+	this->StaticMeshVertexBuffer.BindTangentVertexBuffer(&this->VertexFactory, Data);
+	this->StaticMeshVertexBuffer.BindPackedTexCoordVertexBuffer(&this->VertexFactory, Data);
+	this->ColorVertexBuffer.BindColorVertexBuffer(&this->VertexFactory, Data);
+	this->VertexFactory.SetData(RHICmdList, Data);
+
+	InitOrUpdateResource(RHICmdList, &this->VertexFactory);
+
+	InvalidateRayTracingData();
+	ValidateRayTracingData();		// currently we are immediately validating. This may be revisited in future.
+}
+
+
+void FMeshRenderBufferSet::TransferVertexUpdateToGPU(FRHICommandListBase& RHICmdList, bool bPositions, bool bNormals, bool bTexCoords, bool bColors)
+{
+	if (TriangleCount == 0)
+	{
+		return;
+	}
+
+	if (bPositions)
+	{
+		FPositionVertexBuffer& VertexBuffer = this->PositionVertexBuffer;
+		void* VertexBufferData = RHICmdList.LockBuffer(VertexBuffer.VertexBufferRHI, 0, VertexBuffer.GetNumVertices() * VertexBuffer.GetStride(), RLM_WriteOnly);
+		FMemory::Memcpy(VertexBufferData, VertexBuffer.GetVertexData(), VertexBuffer.GetNumVertices() * VertexBuffer.GetStride());
+		RHICmdList.UnlockBuffer(VertexBuffer.VertexBufferRHI);
+	}
+	if (bNormals)
+	{
+		FStaticMeshVertexBuffer& VertexBuffer = this->StaticMeshVertexBuffer;
+		void* VertexBufferData = RHICmdList.LockBuffer(VertexBuffer.TangentsVertexBuffer.VertexBufferRHI, 0, VertexBuffer.GetTangentSize(), RLM_WriteOnly);
+		FMemory::Memcpy(VertexBufferData, VertexBuffer.GetTangentData(), VertexBuffer.GetTangentSize());
+		RHICmdList.UnlockBuffer(VertexBuffer.TangentsVertexBuffer.VertexBufferRHI);
+	}
+	if (bColors)
+	{
+		FColorVertexBuffer& VertexBuffer = this->ColorVertexBuffer;
+		void* VertexBufferData = RHICmdList.LockBuffer(VertexBuffer.VertexBufferRHI, 0, VertexBuffer.GetNumVertices() * VertexBuffer.GetStride(), RLM_WriteOnly);
+		FMemory::Memcpy(VertexBufferData, VertexBuffer.GetVertexData(), VertexBuffer.GetNumVertices() * VertexBuffer.GetStride());
+		RHICmdList.UnlockBuffer(VertexBuffer.VertexBufferRHI);
+	}
+	if (bTexCoords)
+	{
+		FStaticMeshVertexBuffer& VertexBuffer = this->StaticMeshVertexBuffer;
+		void* VertexBufferData = RHICmdList.LockBuffer(VertexBuffer.TexCoordVertexBuffer.VertexBufferRHI, 0, VertexBuffer.GetTexCoordSize(), RLM_WriteOnly);
+		FMemory::Memcpy(VertexBufferData, VertexBuffer.GetTexCoordData(), VertexBuffer.GetTexCoordSize());
+		RHICmdList.UnlockBuffer(VertexBuffer.TexCoordVertexBuffer.VertexBufferRHI);
+	}
+
+	InvalidateRayTracingData();
+	ValidateRayTracingData();		// currently we are immediately validating. This may be revisited in future.
+}
+
+
+
+/// ---------------  FMeshRenderBufferSetConversionUtil ------------------ /// 
+
+
+void FMeshRenderBufferSetConversionUtil::UpdateSecondaryTriangleBuffer(FMeshRenderBufferSet* RenderBuffers,
+																       const FDynamicMesh3* Mesh,
+																       bool bDuplicate)
+{
+	if (ensure(bUseSecondaryTriBuffers == true && RenderBuffers->Triangles.IsSet()) == false)
+	{
+		return;
+	}
+
+	const TArray<int32>& TriangleIDs = RenderBuffers->Triangles.GetValue();
+	int NumTris = TriangleIDs.Num();
+	TArray<uint32>& Indices = RenderBuffers->IndexBuffer.Indices;
+	TArray<uint32>& SecondaryIndices = RenderBuffers->SecondaryIndexBuffer.Indices;
+
+	RenderBuffers->SecondaryIndexBuffer.Indices.Reset();
+	if (bDuplicate == false)
+	{
+		RenderBuffers->IndexBuffer.Indices.Reset();
+	}
+	for (int k = 0; k < NumTris; ++k)
+	{
+		int TriangleID = TriangleIDs[k];
+		bool bInclude = SecondaryTriFilterFunc(Mesh, TriangleID);
+		if (bInclude)
+		{
+			SecondaryIndices.Add(3 * k);
+			SecondaryIndices.Add(3 * k + 1);
+			SecondaryIndices.Add(3 * k + 2);
+		}
+		else if (bDuplicate == false)
+		{
+			Indices.Add(3 * k);
+			Indices.Add(3 * k + 1);
+			Indices.Add(3 * k + 2);
+		}
+	}
+}
+
+void FMeshRenderBufferSetConversionUtil::RecomputeRenderBufferTriangleIndexSets( FMeshRenderBufferSet* RenderBuffers,
+																		         const FDynamicMesh3* Mesh)
+{
+	if (RenderBuffers->TriangleCount == 0)
+	{
+		return;
+	}
+	if (ensure(RenderBuffers->Triangles.IsSet() && RenderBuffers->Triangles->Num() > 0) == false)
+	{
+		return;
+	}
+
+	//bool bDuplicate = false;		// flag for future use, in case we want to draw all triangles in primary and duplicates in secondary...
+	RenderBuffers->IndexBuffer.Indices.Reset();
+	RenderBuffers->SecondaryIndexBuffer.Indices.Reset();
+
+	TArray<uint32>& Indices = RenderBuffers->IndexBuffer.Indices;
+	TArray<uint32>& SecondaryIndices = RenderBuffers->SecondaryIndexBuffer.Indices;
+	const TArray<int32>& TriangleIDs = RenderBuffers->Triangles.GetValue();
+
+	int NumTris = TriangleIDs.Num();
+	for (int k = 0; k < NumTris; ++k)
+	{
+		int TriangleID = TriangleIDs[k];
+		bool bInclude = SecondaryTriFilterFunc(Mesh, TriangleID);
+		if (bInclude)
+		{
+			SecondaryIndices.Add(3 * k);
+			SecondaryIndices.Add(3 * k + 1);
+			SecondaryIndices.Add(3 * k + 2);
+		}
+		else // if (bDuplicate == false)
+		{
+			Indices.Add(3 * k);
+			Indices.Add(3 * k + 1);
+			Indices.Add(3 * k + 2);
+		}
+	}
+}
+
+
+/// ---------------  ToMeshRenderBufferSet ------------------ /// 
+
+void FDynamicMeshComponentToMeshRenderBufferSet::Convert(UDynamicMeshComponent& DynamicMeshCompnent, FMeshRenderBufferSet& MeshRenderBufferSet, bool bUseComponentSettings)
+{
+	if (bUseComponentSettings)
+	{ 
+		MeshRenderBufferSetConverter.ColorSpaceTransformMode = DynamicMeshCompnent.GetVertexColorSpaceTransformMode();
+		if (DynamicMeshCompnent.GetColorOverrideMode() == EDynamicMeshComponentColorOverrideMode::Constant)
+		{
+			MeshRenderBufferSetConverter.ConstantVertexColor = DynamicMeshCompnent.GetConstantOverrideColor();
+			MeshRenderBufferSetConverter.bIgnoreVertexColors = true;
+		}
+
+		MeshRenderBufferSetConverter.bUsePerTriangleNormals = DynamicMeshCompnent.GetFlatShadingEnabled();
+	}
+
+
+	const FDynamicMesh3* Mesh = DynamicMeshCompnent.GetRenderMesh();
+
+	// find suitable overlays
+	TArray<const FDynamicMeshUVOverlay*, TInlineAllocator<8>> UVOverlays;
+	const FDynamicMeshNormalOverlay* NormalOverlay = nullptr;
+	const FDynamicMeshColorOverlay* ColorOverlay = nullptr;
+	if (Mesh->HasAttributes())
+	{
+		const FDynamicMeshAttributeSet* Attributes = Mesh->Attributes();
+		NormalOverlay = Attributes->PrimaryNormals();
+		UVOverlays.SetNum(Attributes->NumUVLayers());
+		for (int32 k = 0; k < UVOverlays.Num(); ++k)
+		{
+			UVOverlays[k] = Attributes->GetUVLayer(k);
+		}
+		ColorOverlay = Attributes->PrimaryColors();
+	}
+	TUniqueFunction<void(int, int, int, const FVector3f&, FVector3f&, FVector3f&)> TangentsFunc = MakeTangentsFunc(DynamicMeshCompnent);
+
+	MeshRenderBufferSetConverter.InitializeBuffersFromOverlays(&MeshRenderBufferSet, Mesh,
+		Mesh->TriangleCount(), Mesh->TriangleIndicesItr(),
+		UVOverlays, NormalOverlay, ColorOverlay, TangentsFunc);	
+}
+
+
+TUniqueFunction<void(int, int, int, const FVector3f&, FVector3f&, FVector3f&)> FDynamicMeshComponentToMeshRenderBufferSet::MakeTangentsFunc(UDynamicMeshComponent& DynamicMeshComponent, bool bSkipAutoCompute)
+{
+	EDynamicMeshComponentTangentsMode TangentsType = DynamicMeshComponent.GetTangentsType();
+	if (TangentsType == EDynamicMeshComponentTangentsMode::ExternallyProvided)
+	{
+		UE::Geometry::FDynamicMesh3* RenderMesh = DynamicMeshComponent.GetRenderMesh();
+		// If the RenderMesh has tangents, use them. Otherwise we fall back to the orthogonal basis, below.
+		if (RenderMesh && RenderMesh->HasAttributes() && RenderMesh->Attributes()->HasTangentSpace())
+		{
+			UE::Geometry::FDynamicMeshTangents Tangents(RenderMesh);
+			return [Tangents](int VertexID, int TriangleID, int TriVtxIdx, const FVector3f& Normal, FVector3f& TangentX, FVector3f& TangentY) -> void
+				{
+					Tangents.GetTangentVectors(TriangleID, TriVtxIdx, Normal, TangentX, TangentY);
+				};
+		}
+	}
+	else if (TangentsType == EDynamicMeshComponentTangentsMode::AutoCalculated && bSkipAutoCompute == false)
+	{
+		const UE::Geometry::FMeshTangentsf* Tangents = DynamicMeshComponent.GetAutoCalculatedTangents();
+		if (Tangents != nullptr)
+		{
+			return [Tangents](int VertexID, int TriangleID, int TriVtxIdx, const FVector3f& Normal, FVector3f& TangentX, FVector3f& TangentY) -> void
+				{
+					Tangents->GetTriangleVertexTangentVectors(TriangleID, TriVtxIdx, TangentX, TangentY);
+				};
+		}
+	}
+
+	// fallback to orthogonal basis
+	return [](int VertexID, int TriangleID, int TriVtxIdx, const FVector3f& Normal, FVector3f& TangentX, FVector3f& TangentY) -> void
+		{
+			UE::Geometry::VectorUtil::MakePerpVectors(Normal, TangentX, TangentY);
+		};
+}
+
diff --git a/Engine/Source/Runtime/GeometryFramework/Public/Components/BaseDynamicMeshSceneProxy.h b/Engine/Source/Runtime/GeometryFramework/Public/Components/BaseDynamicMeshSceneProxy.h
--- a/Engine/Source/Runtime/GeometryFramework/Public/Components/BaseDynamicMeshSceneProxy.h
+++ b/Engine/Source/Runtime/GeometryFramework/Public/Components/BaseDynamicMeshSceneProxy.h
@@ -4,12 +4,8 @@
 
 #include "CoreMinimal.h"
 #include "PrimitiveSceneProxy.h"
-#include "VertexFactory.h"
-#include "LocalVertexFactory.h"
+#include "MeshRenderBufferSet.h"
 #include "PrimitiveViewRelevance.h"
-#include "Rendering/StaticMeshVertexBuffer.h"
-#include "Rendering/PositionVertexBuffer.h"
-#include "Rendering/ColorVertexBuffer.h"
 #include "DynamicMeshBuilder.h"
 #include "Components/BaseDynamicMeshComponent.h"
 #include "RayTracingGeometry.h"
@@ -32,376 +28,6 @@ class UMaterialInterface;
 struct FRayTracingMaterialGatheringContext;
 class FCardRepresentationData;
 
-/**
- * FMeshRenderBufferSet stores a set of RenderBuffers for a mesh
- */
-class FMeshRenderBufferSet
-{
-public:
-	/** Number of triangles in this renderbuffer set. Note that triangles may be split between IndexBuffer and SecondaryIndexBuffer. */
-	int TriangleCount = 0;
-
-	/** The buffer containing vertex data. */
-	FStaticMeshVertexBuffer StaticMeshVertexBuffer;
-	/** The buffer containing the position vertex data. */
-	FPositionVertexBuffer PositionVertexBuffer;
-	/** The buffer containing the vertex color data. */
-	FColorVertexBuffer ColorVertexBuffer;
-
-	/** triangle indices */
-	FDynamicMeshIndexBuffer32 IndexBuffer;
-
-	/** vertex factory */
-	FLocalVertexFactory VertexFactory;
-
-	/** Material to draw this mesh with */
-	UMaterialInterface* Material = nullptr;
-
-	/**
-	 * Optional list of triangles stored in this buffer. Storing this allows us
-	 * to rebuild the buffers if vertex data changes.
-	 */
-	TOptional<TArray<int>> Triangles;
-
-	/**
-	 * If secondary index buffer is enabled, we populate this index buffer with additional triangles indexing into the same vertex buffers
-	 */
-	bool bEnableSecondaryIndexBuffer = false;
-
-	/**
-	 * partition or subset of IndexBuffer that indexes into same vertex buffers
-	 */
-	FDynamicMeshIndexBuffer32 SecondaryIndexBuffer;
-
-	/**
-	 * configure whether raytracing should be enabled for this RenderBufferSet
-	 */
-	bool bEnableRaytracing = false;
-
-#if RHI_RAYTRACING
-	/**
-	 * Raytracing buffers
-	 */
-	FRayTracingGeometry PrimaryRayTracingGeometry;
-	FRayTracingGeometry SecondaryRayTracingGeometry;
-	bool bIsRayTracingDataValid = false;
-#endif
-
-	/**
-	 * In situations where we want to *update* the existing Vertex or Index buffers, we need to synchronize
-	 * access between the Game and Render threads. We use this lock to do that.
-	 */
-	FCriticalSection BuffersLock;
-
-
-	FMeshRenderBufferSet(ERHIFeatureLevel::Type FeatureLevelType)
-		: VertexFactory(FeatureLevelType, "FMeshRenderBufferSet")
-	{
-		StaticMeshVertexBuffer.SetUseFullPrecisionUVs(true);
-		StaticMeshVertexBuffer.SetUseHighPrecisionTangentBasis(true);
-	}
-
-
-	virtual ~FMeshRenderBufferSet()
-	{
-		if (TriangleCount > 0)
-		{
-			PositionVertexBuffer.ReleaseResource();
-			StaticMeshVertexBuffer.ReleaseResource();
-			ColorVertexBuffer.ReleaseResource();
-			VertexFactory.ReleaseResource();
-			if (IndexBuffer.IsInitialized())
-			{
-				IndexBuffer.ReleaseResource();
-			}
-			if (SecondaryIndexBuffer.IsInitialized())
-			{
-				SecondaryIndexBuffer.ReleaseResource();
-			}
-
-#if RHI_RAYTRACING
-			if (bEnableRaytracing)
-			{
-				PrimaryRayTracingGeometry.ReleaseResource();
-				SecondaryRayTracingGeometry.ReleaseResource();
-			}
-#endif
-		}
-	}
-
-
-	/**
-	 * Upload initialized mesh buffers. 
-	 * @warning This can only be called on the Rendering Thread.
-	 */
-	void Upload()
-	{
-		FRHICommandListBase& RHICmdList = FRHICommandListImmediate::Get();
-
-		if (TriangleCount == 0)
-		{
-			return;
-		}
-
-		InitOrUpdateResource(RHICmdList, &this->PositionVertexBuffer);
-		InitOrUpdateResource(RHICmdList, &this->StaticMeshVertexBuffer);
-		InitOrUpdateResource(RHICmdList, &this->ColorVertexBuffer);
-
-		FLocalVertexFactory::FDataType Data;
-		this->PositionVertexBuffer.BindPositionVertexBuffer(&this->VertexFactory, Data);
-		this->StaticMeshVertexBuffer.BindTangentVertexBuffer(&this->VertexFactory, Data);
-		this->StaticMeshVertexBuffer.BindPackedTexCoordVertexBuffer(&this->VertexFactory, Data);
-		// currently no lightmaps support
-		//this->StaticMeshVertexBuffer.BindLightMapVertexBuffer(&this->VertexFactory, Data, LightMapIndex);
-		this->ColorVertexBuffer.BindColorVertexBuffer(&this->VertexFactory, Data);
-		this->VertexFactory.SetData(RHICmdList, Data);
-
-		InitOrUpdateResource(RHICmdList, &this->VertexFactory);
-		PositionVertexBuffer.InitResource(RHICmdList);
-		StaticMeshVertexBuffer.InitResource(RHICmdList);
-		ColorVertexBuffer.InitResource(RHICmdList);
-		VertexFactory.InitResource(RHICmdList);
-
-		if (IndexBuffer.Indices.Num() > 0)
-		{
-			IndexBuffer.InitResource(RHICmdList);
-		}
-		if (bEnableSecondaryIndexBuffer && SecondaryIndexBuffer.Indices.Num() > 0)
-		{
-			SecondaryIndexBuffer.InitResource(RHICmdList);
-		}
-
-		InvalidateRayTracingData();
-		ValidateRayTracingData();		// currently we are immediately validating. This may be revisited in future.
-	}
-
-
-	/**
-	 * Fast path to only update the primary and secondary index buffers. This can be used
-	 * when (eg) the secondary index buffer is being used to highlight/hide a subset of triangles.
-	 * @warning This can only be called on the Rendering Thread.
-	 */
-	void UploadIndexBufferUpdate()
-	{
-		// todo: can this be done with RHI locking and memcpy, like in TransferVertexUpdateToGPU?
-		FRHICommandListBase& RHICmdList = FRHICommandListImmediate::Get();
-
-		if (IndexBuffer.Indices.Num() > 0)
-		{
-			InitOrUpdateResource(RHICmdList, &IndexBuffer);
-		}
-		if (bEnableSecondaryIndexBuffer && SecondaryIndexBuffer.Indices.Num() > 0)
-		{
-			InitOrUpdateResource(RHICmdList, &SecondaryIndexBuffer);
-		}
-
-		InvalidateRayTracingData();
-		ValidateRayTracingData();		// currently we are immediately validating. This may be revisited in future.
-	}
-
-
-	/**
-	 * Fast path to only update vertex buffers. This path rebuilds all the
-	 * resources and reconfigures the vertex factory, so the counts/etc could be modified.
-	 * @warning This can only be called on the Rendering Thread.
-	 */
-	void UploadVertexUpdate(bool bPositions, bool bMeshAttribs, bool bColors)
-	{
-		// todo: look at calls to this function, it seems possible that TransferVertexUpdateToGPU
-		// could be used instead (which should be somewhat more efficient?). It's not clear if there
-		// are any situations where we would change vertex buffer size w/o also updating the index
-		// buffers (in which case we are fully rebuilding the buffers...)
-
-		if (TriangleCount == 0)
-		{
-			return;
-		}
-
-		FRHICommandListBase& RHICmdList = FRHICommandListImmediate::Get();
-
-		if (bPositions)
-		{
-			InitOrUpdateResource(RHICmdList, &this->PositionVertexBuffer);
-		}
-		if (bMeshAttribs)
-		{
-			InitOrUpdateResource(RHICmdList, &this->StaticMeshVertexBuffer);
-		}
-		if (bColors)
-		{
-			InitOrUpdateResource(RHICmdList, &this->ColorVertexBuffer);
-		}
-
-		FLocalVertexFactory::FDataType Data;
-		this->PositionVertexBuffer.BindPositionVertexBuffer(&this->VertexFactory, Data);
-		this->StaticMeshVertexBuffer.BindTangentVertexBuffer(&this->VertexFactory, Data);
-		this->StaticMeshVertexBuffer.BindPackedTexCoordVertexBuffer(&this->VertexFactory, Data);
-		this->ColorVertexBuffer.BindColorVertexBuffer(&this->VertexFactory, Data);
-		this->VertexFactory.SetData(RHICmdList, Data);
-
-		InitOrUpdateResource(RHICmdList, &this->VertexFactory);
-
-		InvalidateRayTracingData();
-		ValidateRayTracingData();		// currently we are immediately validating. This may be revisited in future.
-	}
-
-
-	/**
-	 * Fast path to update various vertex buffers. This path does not support changing the
-	 * size/counts of any of the sub-buffers, a direct memcopy from the CPU-side buffer to the RHI buffer is used.
-	 * @warning This can only be called on the Rendering Thread.
-	 */
-	void TransferVertexUpdateToGPU(FRHICommandListBase& RHICmdList, bool bPositions, bool bNormals, bool bTexCoords, bool bColors)
-	{
-		if (TriangleCount == 0)
-		{
-			return;
-		}
-
-		if (bPositions)
-		{
-			FPositionVertexBuffer& VertexBuffer = this->PositionVertexBuffer;
-			void* VertexBufferData = RHICmdList.LockBuffer(VertexBuffer.VertexBufferRHI, 0, VertexBuffer.GetNumVertices() * VertexBuffer.GetStride(), RLM_WriteOnly);
-			FMemory::Memcpy(VertexBufferData, VertexBuffer.GetVertexData(), VertexBuffer.GetNumVertices() * VertexBuffer.GetStride());
-			RHICmdList.UnlockBuffer(VertexBuffer.VertexBufferRHI);
-		}
-		if (bNormals)
-		{
-			FStaticMeshVertexBuffer& VertexBuffer = this->StaticMeshVertexBuffer;
-			void* VertexBufferData = RHICmdList.LockBuffer(VertexBuffer.TangentsVertexBuffer.VertexBufferRHI, 0, VertexBuffer.GetTangentSize(), RLM_WriteOnly);
-			FMemory::Memcpy(VertexBufferData, VertexBuffer.GetTangentData(), VertexBuffer.GetTangentSize());
-			RHICmdList.UnlockBuffer(VertexBuffer.TangentsVertexBuffer.VertexBufferRHI);
-		}
-		if (bColors)
-		{
-			FColorVertexBuffer& VertexBuffer = this->ColorVertexBuffer;
-			void* VertexBufferData = RHICmdList.LockBuffer(VertexBuffer.VertexBufferRHI, 0, VertexBuffer.GetNumVertices() * VertexBuffer.GetStride(), RLM_WriteOnly);
-			FMemory::Memcpy(VertexBufferData, VertexBuffer.GetVertexData(), VertexBuffer.GetNumVertices() * VertexBuffer.GetStride());
-			RHICmdList.UnlockBuffer(VertexBuffer.VertexBufferRHI);
-		}
-		if (bTexCoords)
-		{
-			FStaticMeshVertexBuffer& VertexBuffer = this->StaticMeshVertexBuffer;
-			void* VertexBufferData = RHICmdList.LockBuffer(VertexBuffer.TexCoordVertexBuffer.VertexBufferRHI, 0, VertexBuffer.GetTexCoordSize(), RLM_WriteOnly);
-			FMemory::Memcpy(VertexBufferData, VertexBuffer.GetTexCoordData(), VertexBuffer.GetTexCoordSize());
-			RHICmdList.UnlockBuffer(VertexBuffer.TexCoordVertexBuffer.VertexBufferRHI);
-		}
-
-		InvalidateRayTracingData();
-		ValidateRayTracingData();		// currently we are immediately validating. This may be revisited in future.
-	}
-
-	UE_DEPRECATED(5.3, "TransferVertexUpdateToGPU now requires a command list")
-	void TransferVertexUpdateToGPU(bool bPositions, bool bNormals, bool bTexCoords, bool bColors)
-	{
-		TransferVertexUpdateToGPU(FRHICommandListImmediate::Get(), bPositions, bNormals, bTexCoords, bColors);
-	}
-
-	void InvalidateRayTracingData()
-	{
-#if RHI_RAYTRACING
-		bIsRayTracingDataValid = false;
-#endif
-	}
-
-	// Verify that valid raytracing data is available. This will cause a rebuild of the
-	// raytracing data if any of our buffers have been modified. Currently this is called
-	// by GetDynamicRayTracingInstances to ensure the RT data is available when needed.
-	void ValidateRayTracingData()
-	{
-#if RHI_RAYTRACING
-		if (bIsRayTracingDataValid == false && IsRayTracingEnabled() && bEnableRaytracing)
-		{
-			UpdateRaytracingGeometryIfEnabled();
-
-			bIsRayTracingDataValid = true;
-		}
-#endif
-	}
-
-
-protected:
-
-	// rebuild raytracing data for current buffers
-	void UpdateRaytracingGeometryIfEnabled()
-	{
-#if RHI_RAYTRACING
-		// do we always want to do this?
-		PrimaryRayTracingGeometry.ReleaseResource();		
-		SecondaryRayTracingGeometry.ReleaseResource();
-		FRHICommandListBase& RHICmdList = FRHICommandListImmediate::Get();
-			
-		for (int32 k = 0; k < 2; ++k)
-		{
-			FDynamicMeshIndexBuffer32& UseIndexBuffer = (k == 0) ? IndexBuffer : SecondaryIndexBuffer;
-			if (UseIndexBuffer.Indices.Num() == 0)
-			{
-				continue;
-			}
-
-			FRayTracingGeometry& RayTracingGeometry = (k == 0) ? PrimaryRayTracingGeometry : SecondaryRayTracingGeometry;
-
-			FRayTracingGeometryInitializer Initializer;
-			Initializer.IndexBuffer = UseIndexBuffer.IndexBufferRHI;
-			Initializer.TotalPrimitiveCount = UseIndexBuffer.Indices.Num() / 3;
-			Initializer.GeometryType = RTGT_Triangles;
-			Initializer.bFastBuild = true;
-			Initializer.bAllowUpdate = false;
-
-			FRayTracingGeometrySegment Segment;
-			Segment.VertexBuffer = PositionVertexBuffer.VertexBufferRHI;
-			Segment.NumPrimitives = Initializer.TotalPrimitiveCount;
-			Segment.MaxVertices = PositionVertexBuffer.GetNumVertices();
-
-			Initializer.Segments.Add(Segment);
-
-			RayTracingGeometry.SetInitializer(MoveTemp(Initializer));
-			RayTracingGeometry.InitResource(RHICmdList);
-		}
-#endif
-	}
-
-
-
-	/**
-	 * Initializes a render resource, or update it if already initialized.
-	 * @warning This function can only be called on the Render Thread
-	 */
-	void InitOrUpdateResource(FRHICommandListBase& RHICmdList, FRenderResource* Resource)
-	{
-		if (!Resource->IsInitialized())
-		{
-			Resource->InitResource(RHICmdList);
-		}
-		else
-		{
-			Resource->UpdateRHI(RHICmdList);
-		}
-	}
-
-
-
-
-protected:
-	friend class FBaseDynamicMeshSceneProxy;
-
-	/**
-	 * Enqueue a command on the Render Thread to destroy the passed in buffer set.
-	 * At this point the buffer set should be considered invalid.
-	 */
-	static void DestroyRenderBufferSet(FMeshRenderBufferSet* BufferSet)
-	{
-		if (BufferSet->TriangleCount == 0)
-		{
-			return;
-		}
-
-		delete BufferSet;
-	}
-
-
-};
-
 
 
 
@@ -418,56 +44,11 @@ public:
 	UBaseDynamicMeshComponent* ParentBaseComponent;
 
 	/**
-	 * Constant color assigned to vertices if no other vertex color is specified
-	 */
-	FColor ConstantVertexColor = FColor::White;
-
-	/**
-	 * If true, vertex colors on the FDynamicMesh will be ignored
-	 */
-	bool bIgnoreVertexColors = false;
-
-	/**
-	 * If true, a per-triangle color is used to set vertex colors
-	 */
-	bool bUsePerTriangleColor = false;
-
-	/**
-	 * Per-triangle color function. Only called if bUsePerTriangleColor=true
-	 */
-	TFunction<FColor(const FDynamicMesh3*, int)> PerTriangleColorFunc = nullptr;
-
-
-	/**
-	 * If true, VertexColorRemappingFunc is called on Vertex Colors provided from Mesh to remap them to a different color
-	 */
-	bool bApplyVertexColorRemapping = false;
-
-	/**
-	 * Vertex color remapping function. Only called if bApplyVertexColorRemapping == true, for mesh vertex colors
-	 */
-	TUniqueFunction<void(FVector4f&)> VertexColorRemappingFunc = nullptr;
-
-	/**
-	 * Color Space Transform/Conversion applied to Vertex Colors provided from Mesh Color Overlay Attribute
-	 * Color Space Conversion is applied after any Vertex Color Remapping.
-	 */
-	EDynamicMeshVertexColorTransformMode ColorSpaceTransformMode = EDynamicMeshVertexColorTransformMode::NoTransform;
-
-	/**
-	* If true, a facet normals are used instead of mesh normals
+	* Utility to initialize and update the mesh render buffers from a mesh with overlays
+	* and holds all the settings required.
 	*/
-	bool bUsePerTriangleNormals = false;
+	FMeshRenderBufferSetConversionUtil MeshRenderBufferSetConverter;
 
-	/**
-	 * If true, populate secondary buffers using SecondaryTriFilterFunc
-	 */
-	bool bUseSecondaryTriBuffers = false;
-
-	/**
-	 * Filter predicate for secondary triangle index buffer. Only called if bUseSecondaryTriBuffers=true
-	 */
-	TUniqueFunction<bool(const FDynamicMesh3*, int32)> SecondaryTriFilterFunc = nullptr;
 
 protected:
 	// Set of currently-allocated RenderBuffers. We own these pointers and must clean them up.
@@ -537,10 +118,7 @@ public:
 		TFunctionRef<void(int, int, int, const FVector3f&, FVector3f&, FVector3f&)> TangentsFunc,
 		bool bTrackTriangles = false)
 	{
-		TArray<const FDynamicMeshUVOverlay*> UVOverlays;
-		UVOverlays.Add(UVOverlay);
-		InitializeBuffersFromOverlays(RenderBuffers, Mesh, NumTriangles, Enumerable,
-			UVOverlays, NormalOverlay, ColorOverlay, TangentsFunc, bTrackTriangles);
+		MeshRenderBufferSetConverter.InitializeBuffersFromOverlays(RenderBuffers, Mesh, NumTriangles, Enumerable, UVOverlay, NormalOverlay, ColorOverlay, TangentsFunc, bTrackTriangles);
 	}
 
 
@@ -560,128 +138,10 @@ public:
 		TFunctionRef<void(int, int, int, const FVector3f&, FVector3f&, FVector3f&)> TangentsFunc,
 		bool bTrackTriangles = false)
 	{
-		RenderBuffers->TriangleCount = NumTriangles;
-		if (NumTriangles == 0)
-		{
-			return;
-		}
-
-		bool bHaveColors = (ColorOverlay != nullptr) && (bIgnoreVertexColors == false);
-
-		int NumVertices = NumTriangles * 3;
-		int NumUVOverlays = UVOverlays.Num();
-		int NumTexCoords = FMath::Max(1, NumUVOverlays);		// must have at least one tex coord
-		TArray<FIndex3i, TFixedAllocator<MAX_STATIC_TEXCOORDS>> UVTriangles;
-		UVTriangles.SetNum(NumTexCoords);
-
-		{
-			RenderBuffers->PositionVertexBuffer.Init(NumVertices);
-			RenderBuffers->StaticMeshVertexBuffer.Init(NumVertices, NumTexCoords );
-			RenderBuffers->ColorVertexBuffer.Init(NumVertices);
-			RenderBuffers->IndexBuffer.Indices.AddUninitialized(NumTriangles * 3);
-		}
-
-		// build triangle list if requested, or if we are using secondary buffers in which case we need it to filter later
-		bool bBuildTriangleList = bTrackTriangles || bUseSecondaryTriBuffers;
-		if (bBuildTriangleList)
-		{
-			RenderBuffers->Triangles = TArray<int32>();
-		}
-
-		int TriIdx = 0, VertIdx = 0;
-		FVector3f TangentX, TangentY;
-		for (int TriangleID : Enumerable)
-		{
-			FIndex3i Tri = Mesh->GetTriangle(TriangleID);
-			for (int32 k = 0; k < NumTexCoords; ++k)
-			{
-				UVTriangles[k] = (k < NumUVOverlays && UVOverlays[k] != nullptr) ? UVOverlays[k]->GetTriangle(TriangleID) : FIndex3i::Invalid();
-			}
-			FIndex3i TriNormal = (NormalOverlay != nullptr) ? NormalOverlay->GetTriangle(TriangleID) : FIndex3i::Zero();
-			FIndex3i TriColor = (ColorOverlay != nullptr) ? ColorOverlay->GetTriangle(TriangleID) : FIndex3i::Zero();
-
-			FColor UniformTriColor = ConstantVertexColor;
-			if (bUsePerTriangleColor && PerTriangleColorFunc != nullptr)
-			{
-				UniformTriColor = PerTriangleColorFunc(Mesh, TriangleID);
-				bHaveColors = false;
-			}
-
-			for (int j = 0; j < 3; ++j)
-			{
-				RenderBuffers->PositionVertexBuffer.VertexPosition(VertIdx) = (FVector3f)Mesh->GetVertex(Tri[j]);
-
-				FVector3f Normal;
-				if (bUsePerTriangleNormals)
-				{
-					Normal = (FVector3f)Mesh->GetTriNormal(TriangleID);
-				}
-				else
-				{
-					Normal = (NormalOverlay != nullptr && TriNormal[j] != FDynamicMesh3::InvalidID) ?
-						NormalOverlay->GetElement(TriNormal[j]) : Mesh->GetVertexNormal(Tri[j]);
-				}
-
-				// get tangents
-				TangentsFunc(Tri[j], TriangleID, j, Normal, TangentX, TangentY);
-
-				RenderBuffers->StaticMeshVertexBuffer.SetVertexTangents(VertIdx, TangentX, TangentY, Normal);
-
-				for (int32 k = 0; k < NumTexCoords; ++k)
-				{
-					FVector2f UV = (UVTriangles[k][j] != FDynamicMesh3::InvalidID) ?
-						UVOverlays[k]->GetElement(UVTriangles[k][j]) : FVector2f::Zero();
-					RenderBuffers->StaticMeshVertexBuffer.SetVertexUV(VertIdx, k, UV);
-				}
-
-				FColor VertexFColor = (bHaveColors && TriColor[j] != FDynamicMesh3::InvalidID) ?
-					GetOverlayColorAsFColor(ColorOverlay, TriColor[j]) : UniformTriColor;
-
-				RenderBuffers->ColorVertexBuffer.VertexColor(VertIdx) = VertexFColor;
-
-				RenderBuffers->IndexBuffer.Indices[TriIdx++] = VertIdx;		// currently TriIdx == VertIdx so we don't really need both...
-				VertIdx++;
-			}
-
-			if (bBuildTriangleList)
-			{
-				RenderBuffers->Triangles->Add(TriangleID);
-			}
-		}
-
-		// split triangles into secondary buffer (at bit redudant since we just built IndexBuffer, but we may optionally duplicate triangles in the future
-		if (bUseSecondaryTriBuffers)
-		{
-			RenderBuffers->bEnableSecondaryIndexBuffer = true;
-			UpdateSecondaryTriangleBuffer(RenderBuffers, Mesh, false);
-		}
+		MeshRenderBufferSetConverter.InitializeBuffersFromOverlays(RenderBuffers, Mesh, NumTriangles, Enumerable, UVOverlays, NormalOverlay, ColorOverlay, TangentsFunc, bTrackTriangles);
 	}
 
 
-	FColor GetOverlayColorAsFColor(
-		const FDynamicMeshColorOverlay* ColorOverlay,
-		int32 ElementID)
-	{
-		checkSlow(ColorOverlay);
-		FVector4f UseColor = ColorOverlay->GetElement(ElementID);
-
-		if (bApplyVertexColorRemapping)
-		{
-			VertexColorRemappingFunc(UseColor);
-		}
-
-		if (ColorSpaceTransformMode == EDynamicMeshVertexColorTransformMode::SRGBToLinear)
-		{
-			// is there a better way to do this? 
-			FColor QuantizedSRGBColor = ((FLinearColor)UseColor).ToFColor(false);
-			return FLinearColor(QuantizedSRGBColor).ToFColor(false);
-		}
-		else
-		{
-			bool bConvertToSRGB = (ColorSpaceTransformMode == EDynamicMeshVertexColorTransformMode::LinearToSRGB);
-			return ((FLinearColor)UseColor).ToFColor(bConvertToSRGB);
-		}
-	}
 
 
 	/**
@@ -694,38 +154,7 @@ public:
 		const FDynamicMesh3* Mesh,
 		bool bDuplicate)
 	{
-		if (ensure(bUseSecondaryTriBuffers == true && RenderBuffers->Triangles.IsSet()) == false)
-		{
-			return;
-		}
-
-		const TArray<int32>& TriangleIDs = RenderBuffers->Triangles.GetValue();
-		int NumTris = TriangleIDs.Num();
-		TArray<uint32>& Indices = RenderBuffers->IndexBuffer.Indices;
-		TArray<uint32>& SecondaryIndices = RenderBuffers->SecondaryIndexBuffer.Indices;
-
-		RenderBuffers->SecondaryIndexBuffer.Indices.Reset();
-		if (bDuplicate == false)
-		{
-			RenderBuffers->IndexBuffer.Indices.Reset();
-		}
-		for ( int k = 0; k < NumTris; ++k)
-		{
-			int TriangleID = TriangleIDs[k];
-			bool bInclude = SecondaryTriFilterFunc(Mesh, TriangleID);
-			if (bInclude)
-			{
-				SecondaryIndices.Add(3*k);
-				SecondaryIndices.Add(3*k + 1);
-				SecondaryIndices.Add(3*k + 2);
-			} 
-			else if (bDuplicate == false)
-			{
-				Indices.Add(3*k);
-				Indices.Add(3*k + 1);
-				Indices.Add(3*k + 2);
-			}
-		}
+		MeshRenderBufferSetConverter.UpdateSecondaryTriangleBuffer(RenderBuffers, Mesh, bDuplicate);
 	}
 
 
@@ -735,44 +164,10 @@ public:
 	 * after this function!
 	 */
 	void RecomputeRenderBufferTriangleIndexSets(
-		FMeshRenderBufferSet* RenderBuffers, 
+		FMeshRenderBufferSet* RenderBuffers,
 		const FDynamicMesh3* Mesh)
 	{
-		if (RenderBuffers->TriangleCount == 0)
-		{
-			return;
-		}
-		if (ensure(RenderBuffers->Triangles.IsSet() && RenderBuffers->Triangles->Num() > 0) == false)
-		{
-			return;
-		}
-
-		//bool bDuplicate = false;		// flag for future use, in case we want to draw all triangles in primary and duplicates in secondary...
-		RenderBuffers->IndexBuffer.Indices.Reset();
-		RenderBuffers->SecondaryIndexBuffer.Indices.Reset();
-		
-		TArray<uint32>& Indices = RenderBuffers->IndexBuffer.Indices;
-		TArray<uint32>& SecondaryIndices = RenderBuffers->SecondaryIndexBuffer.Indices;
-		const TArray<int32>& TriangleIDs = RenderBuffers->Triangles.GetValue();
-
-		int NumTris = TriangleIDs.Num();
-		for (int k = 0; k < NumTris; ++k)
-		{
-			int TriangleID = TriangleIDs[k];
-			bool bInclude = SecondaryTriFilterFunc(Mesh, TriangleID);
-			if (bInclude)
-			{
-				SecondaryIndices.Add(3 * k);
-				SecondaryIndices.Add(3 * k + 1);
-				SecondaryIndices.Add(3 * k + 2);
-			}
-			else // if (bDuplicate == false)
-			{
-				Indices.Add(3 * k);
-				Indices.Add(3 * k + 1);
-				Indices.Add(3 * k + 2);
-			}
-		}
+		MeshRenderBufferSetConverter.RecomputeRenderBufferTriangleIndexSets(RenderBuffers, Mesh);
 	}
 
 
@@ -793,73 +188,7 @@ public:
 		bool bUpdateNormals = false,
 		bool bUpdateColors = false)
 	{
-		if (RenderBuffers->TriangleCount == 0)
-		{
-			return;
-		}
-
-		bool bHaveColors = (ColorOverlay != nullptr) && (bIgnoreVertexColors == false);
-
-		int NumVertices = NumTriangles * 3;
-		if ( (bUpdatePositions && ensure(RenderBuffers->PositionVertexBuffer.GetNumVertices() == NumVertices) == false )
-			|| (bUpdateNormals && ensure(RenderBuffers->StaticMeshVertexBuffer.GetNumVertices() == NumVertices) == false )
-			|| (bUpdateColors && ensure(RenderBuffers->ColorVertexBuffer.GetNumVertices() == NumVertices) == false ) )
-		{
-			return;
-		}
-
-		int VertIdx = 0;
-		FVector3f TangentX, TangentY;
-		for (int TriangleID : Enumerable)
-		{
-			FIndex3i Tri = Mesh->GetTriangle(TriangleID);
-
-			FIndex3i TriNormal = (bUpdateNormals && NormalOverlay != nullptr) ? NormalOverlay->GetTriangle(TriangleID) : FIndex3i::Zero();
-			FIndex3i TriColor = (bUpdateColors && ColorOverlay != nullptr) ? ColorOverlay->GetTriangle(TriangleID) : FIndex3i::Zero();
-
-			FColor UniformTriColor = ConstantVertexColor;
-			if (bUpdateColors && bUsePerTriangleColor && PerTriangleColorFunc != nullptr)
-			{
-				UniformTriColor = PerTriangleColorFunc(Mesh, TriangleID);
-				bHaveColors = false;
-			}
-
-			for (int j = 0; j < 3; ++j)
-			{
-				if (bUpdatePositions)
-				{
-					RenderBuffers->PositionVertexBuffer.VertexPosition(VertIdx) = (FVector3f)Mesh->GetVertex(Tri[j]);
-				}
-
-				if (bUpdateNormals)
-				{
-					// get normal and tangent
-					FVector3f Normal;
-					if (bUsePerTriangleNormals)
-					{
-						Normal = (FVector3f)Mesh->GetTriNormal(TriangleID);
-					}
-					else
-					{
-						Normal = (NormalOverlay != nullptr && TriNormal[j] != FDynamicMesh3::InvalidID) ?
-							NormalOverlay->GetElement(TriNormal[j]) : Mesh->GetVertexNormal(Tri[j]);
-					}
-
-					TangentsFunc(Tri[j], TriangleID, j, Normal, TangentX, TangentY);
-
-					RenderBuffers->StaticMeshVertexBuffer.SetVertexTangents(VertIdx, (FVector3f)TangentX, (FVector3f)TangentY, (FVector3f)Normal);
-				}
-
-				if (bUpdateColors)
-				{
-					FColor VertexFColor = (bHaveColors  && TriColor[j] != FDynamicMesh3::InvalidID) ?
-						GetOverlayColorAsFColor(ColorOverlay, TriColor[j]) : UniformTriColor;
-					RenderBuffers->ColorVertexBuffer.VertexColor(VertIdx) = VertexFColor;
-				}
-
-				VertIdx++;
-			}
-		}
+		MeshRenderBufferSetConverter.UpdateVertexBuffersFromOverlays(RenderBuffers, Mesh, NumTriangles, Enumerable, NormalOverlay, ColorOverlay, TangentsFunc, bUpdatePositions, bUpdateNormals, bUpdateColors);
 	}
 
 	/**
@@ -873,49 +202,7 @@ public:
 		int32 NumTriangles, TriangleEnumerable Enumerable,
 		const TArray<const FDynamicMeshUVOverlay*, UVOverlayListAllocator>& UVOverlays)
 	{
-		// We align the update to the way we set UV's in InitializeBuffersFromOverlays.
-
-		if (RenderBuffers->TriangleCount == 0)
-		{
-			return;
-		}
-		int NumVertices = NumTriangles * 3;
-		if (ensure(RenderBuffers->StaticMeshVertexBuffer.GetNumVertices() == NumVertices) == false)
-		{
-			return;
-		}
-
-		int NumUVOverlays = UVOverlays.Num();
-		int NumTexCoords = RenderBuffers->StaticMeshVertexBuffer.GetNumTexCoords();
-		if (!ensure(NumUVOverlays <= NumTexCoords))
-		{
-			return;
-		}
-
-		// Temporarily stores the UV element indices for all UV channels of a single triangle
-		TArray<FIndex3i, TFixedAllocator<MAX_STATIC_TEXCOORDS>> UVTriangles;
-		UVTriangles.SetNum(NumTexCoords);
-
-		int VertIdx = 0;
-		for (int TriangleID : Enumerable)
-		{
-			for (int32 k = 0; k < NumTexCoords; ++k)
-			{
-				UVTriangles[k] = (k < NumUVOverlays && UVOverlays[k] != nullptr) ? UVOverlays[k]->GetTriangle(TriangleID) : FIndex3i::Invalid();
-			}
-
-			for (int j = 0; j < 3; ++j)
-			{
-				for (int32 k = 0; k < NumTexCoords; ++k)
-				{
-					FVector2f UV = (UVTriangles[k][j] != FDynamicMesh3::InvalidID) ?
-						UVOverlays[k]->GetElement(UVTriangles[k][j]) : FVector2f::Zero();
-					RenderBuffers->StaticMeshVertexBuffer.SetVertexUV(VertIdx, k, UV);
-				}
-
-				++VertIdx;
-			}
-		}
+		MeshRenderBufferSetConverter.UpdateVertexUVBufferFromOverlays(RenderBuffers, Mesh, NumTriangles, Enumerable, UVOverlays);
 	}
 
 
diff --git a/Engine/Source/Runtime/GeometryFramework/Public/Components/MeshRenderBufferSet.h b/Engine/Source/Runtime/GeometryFramework/Public/Components/MeshRenderBufferSet.h
new file mode 100644
--- /dev/null
+++ b/Engine/Source/Runtime/GeometryFramework/Public/Components/MeshRenderBufferSet.h
@@ -0,0 +1,738 @@
+// Copyright Epic Games, Inc. All Rights Reserved. 
+
+#pragma once
+
+#include "CoreMinimal.h"
+
+#include "Components/BaseDynamicMeshComponent.h"
+#include "Components/DynamicMeshComponent.h"
+#include "DynamicMeshBuilder.h"
+#include "DynamicMesh/DynamicMesh3.h"	
+#include "DynamicMesh/DynamicMeshAttributeSet.h"
+
+#include "LocalVertexFactory.h"
+#include "Rendering/StaticMeshVertexBuffer.h"
+#include "Rendering/PositionVertexBuffer.h"
+#include "Rendering/ColorVertexBuffer.h"
+#include "RHICommandList.h"
+#include "RayTracingGeometry.h"
+#include "VertexFactory.h"
+
+
+
+
+using UE::Geometry::FDynamicMesh3;
+using UE::Geometry::FDynamicMeshAttributeSet;
+using UE::Geometry::FDynamicMeshUVOverlay;
+using UE::Geometry::FDynamicMeshNormalOverlay;
+using UE::Geometry::FDynamicMeshColorOverlay;
+using UE::Geometry::FDynamicMeshMaterialAttribute;
+
+class UMaterialInterface;
+
+/**
+ * FMeshRenderBufferSet stores a set of RenderBuffers for a mesh
+ */
+class FMeshRenderBufferSet
+{
+public:
+	/** Number of triangles in this renderbuffer set. Note that triangles may be split between IndexBuffer and SecondaryIndexBuffer. */
+	int TriangleCount = 0;
+
+	/** The buffer containing vertex data. */
+	FStaticMeshVertexBuffer StaticMeshVertexBuffer;
+	/** The buffer containing the position vertex data. */
+	FPositionVertexBuffer PositionVertexBuffer;
+	/** The buffer containing the vertex color data. */
+	FColorVertexBuffer ColorVertexBuffer;
+
+	/** triangle indices */
+	FDynamicMeshIndexBuffer32 IndexBuffer;
+
+	/** vertex factory */
+	FLocalVertexFactory VertexFactory;
+
+	/** Material to draw this mesh with */
+	UMaterialInterface* Material = nullptr;
+
+	/**
+	 * Optional list of triangles stored in this buffer. Storing this allows us
+	 * to rebuild the buffers if vertex data changes.
+	 */
+	TOptional<TArray<int>> Triangles;
+
+	/**
+	 * If secondary index buffer is enabled, we populate this index buffer with additional triangles indexing into the same vertex buffers
+	 */
+	bool bEnableSecondaryIndexBuffer = false;
+
+	/**
+	 * partition or subset of IndexBuffer that indexes into same vertex buffers
+	 */
+	FDynamicMeshIndexBuffer32 SecondaryIndexBuffer;
+
+	/**
+	 * configure whether raytracing should be enabled for this RenderBufferSet
+	 */
+	bool bEnableRaytracing = false;
+
+#if RHI_RAYTRACING
+	/**
+	 * Raytracing buffers
+	 */
+	FRayTracingGeometry PrimaryRayTracingGeometry;
+	FRayTracingGeometry SecondaryRayTracingGeometry;
+	bool bIsRayTracingDataValid = false;
+#endif
+
+	/**
+	 * In situations where we want to *update* the existing Vertex or Index buffers, we need to synchronize
+	 * access between the Game and Render threads. We use this lock to do that.
+	 */
+	FCriticalSection BuffersLock;
+
+
+	FMeshRenderBufferSet(ERHIFeatureLevel::Type FeatureLevelType)
+	: VertexFactory(FeatureLevelType, "FMeshRenderBufferSet")
+	{
+		StaticMeshVertexBuffer.SetUseFullPrecisionUVs(true);
+		StaticMeshVertexBuffer.SetUseHighPrecisionTangentBasis(true);
+	}
+
+	virtual ~FMeshRenderBufferSet()
+	{
+		if (TriangleCount > 0)
+		{
+			PositionVertexBuffer.ReleaseResource();
+			StaticMeshVertexBuffer.ReleaseResource();
+			ColorVertexBuffer.ReleaseResource();
+			VertexFactory.ReleaseResource();
+			if (IndexBuffer.IsInitialized())
+			{
+				IndexBuffer.ReleaseResource();
+			}
+			if (SecondaryIndexBuffer.IsInitialized())
+			{
+				SecondaryIndexBuffer.ReleaseResource();
+			}
+
+#if RHI_RAYTRACING
+			if (bEnableRaytracing)
+			{
+				PrimaryRayTracingGeometry.ReleaseResource();
+				SecondaryRayTracingGeometry.ReleaseResource();
+			}
+#endif
+		}
+	}
+
+
+	/**
+	 * Upload initialized mesh buffers.
+	 * @warning This can only be called on the Rendering Thread.
+	 */
+	GEOMETRYFRAMEWORK_API void Upload();
+	
+
+
+	/**
+	 * Fast path to only update the primary and secondary index buffers. This can be used
+	 * when (eg) the secondary index buffer is being used to highlight/hide a subset of triangles.
+	 * @warning This can only be called on the Rendering Thread.
+	 */
+	void UploadIndexBufferUpdate()
+	{
+		// todo: can this be done with RHI locking and memcpy, like in TransferVertexUpdateToGPU?
+		FRHICommandListBase& RHICmdList = FRHICommandListImmediate::Get();
+
+		if (IndexBuffer.Indices.Num() > 0)
+		{
+			InitOrUpdateResource(RHICmdList, &IndexBuffer);
+		}
+		if (bEnableSecondaryIndexBuffer && SecondaryIndexBuffer.Indices.Num() > 0)
+		{
+			InitOrUpdateResource(RHICmdList, &SecondaryIndexBuffer);
+		}
+
+		InvalidateRayTracingData();
+		ValidateRayTracingData();		// currently we are immediately validating. This may be revisited in future.
+	}
+
+
+	/**
+	 * Fast path to only update vertex buffers. This path rebuilds all the
+	 * resources and reconfigures the vertex factory, so the counts/etc could be modified.
+	 * @warning This can only be called on the Rendering Thread.
+	 */
+	GEOMETRYFRAMEWORK_API void UploadVertexUpdate(bool bPositions, bool bMeshAttribs, bool bColors);
+	
+
+
+	/**
+	 * Fast path to update various vertex buffers. This path does not support changing the
+	 * size/counts of any of the sub-buffers, a direct memcopy from the CPU-side buffer to the RHI buffer is used.
+	 * @warning This can only be called on the Rendering Thread.
+	 */
+	GEOMETRYFRAMEWORK_API void TransferVertexUpdateToGPU(FRHICommandListBase& RHICmdList, bool bPositions, bool bNormals, bool bTexCoords, bool bColors);
+
+
+	void InvalidateRayTracingData()
+	{
+#if RHI_RAYTRACING
+		bIsRayTracingDataValid = false;
+#endif
+	}
+
+	// Verify that valid raytracing data is available. This will cause a rebuild of the
+	// raytracing data if any of our buffers have been modified. Currently this is called
+	// by GetDynamicRayTracingInstances to ensure the RT data is available when needed.
+	void ValidateRayTracingData()
+	{
+#if RHI_RAYTRACING
+		if (bIsRayTracingDataValid == false && IsRayTracingEnabled() && bEnableRaytracing)
+		{
+			UpdateRaytracingGeometryIfEnabled();
+
+			bIsRayTracingDataValid = true;
+		}
+#endif
+	}
+
+
+protected:
+
+	// rebuild raytracing data for current buffers
+	void UpdateRaytracingGeometryIfEnabled()
+	{
+#if RHI_RAYTRACING
+		// do we always want to do this?
+		PrimaryRayTracingGeometry.ReleaseResource();
+		SecondaryRayTracingGeometry.ReleaseResource();
+		FRHICommandListBase& RHICmdList = FRHICommandListImmediate::Get();
+
+		for (int32 k = 0; k < 2; ++k)
+		{
+			FDynamicMeshIndexBuffer32& UseIndexBuffer = (k == 0) ? IndexBuffer : SecondaryIndexBuffer;
+			if (UseIndexBuffer.Indices.Num() == 0)
+			{
+				continue;
+			}
+
+			FRayTracingGeometry& RayTracingGeometry = (k == 0) ? PrimaryRayTracingGeometry : SecondaryRayTracingGeometry;
+
+			FRayTracingGeometryInitializer Initializer;
+			Initializer.IndexBuffer = UseIndexBuffer.IndexBufferRHI;
+			Initializer.TotalPrimitiveCount = UseIndexBuffer.Indices.Num() / 3;
+			Initializer.GeometryType = RTGT_Triangles;
+			Initializer.bFastBuild = true;
+			Initializer.bAllowUpdate = false;
+
+			FRayTracingGeometrySegment Segment;
+			Segment.VertexBuffer = PositionVertexBuffer.VertexBufferRHI;
+			Segment.NumPrimitives = Initializer.TotalPrimitiveCount;
+			Segment.MaxVertices = PositionVertexBuffer.GetNumVertices();
+
+			Initializer.Segments.Add(Segment);
+
+			RayTracingGeometry.SetInitializer(MoveTemp(Initializer));
+			RayTracingGeometry.InitResource(RHICmdList);
+		}
+#endif
+	}
+
+
+
+	/**
+	 * Initializes a render resource, or update it if already initialized.
+	 * @warning This function can only be called on the Render Thread
+	 */
+	void InitOrUpdateResource(FRHICommandListBase& RHICmdList, FRenderResource* Resource)
+	{
+		if (!Resource->IsInitialized())
+		{
+			Resource->InitResource(RHICmdList);
+		}
+		else
+		{
+			Resource->UpdateRHI(RHICmdList);
+		}
+	}
+
+
+public: 
+
+	/**
+	 * Enqueue a command on the Render Thread to destroy the passed in buffer set.
+	 * At this point the buffer set should be considered invalid.
+	 */
+	static void DestroyRenderBufferSet(FMeshRenderBufferSet* BufferSet)
+	{
+		//@todo, remove this check?  based on our usage it could produce a memory leak - see BaseDynamicMeshSceneProxy.cpp
+		if (BufferSet->TriangleCount == 0)
+		{
+			return;
+		}
+
+		delete BufferSet;
+	}
+
+
+};
+
+
+/**
+* Parameters and accompanying functions used to build and update FMeshRenderBuffers from a dynamic mesh.
+* For usage, see FBaseDynamicMeshSceneProxy or the simple converter below FDynamicMeshComponentToMeshRenderBufferSet
+* 
+* Note: This is low-level and methods assume, but do not check, all input data is consistent (e.g. that all overlays, 
+        the triangle count and enumerator are all consistent with the provided mesh. )
+*/
+struct FMeshRenderBufferSetConversionUtil
+{
+
+	/**
+	* Constant color assigned to vertices if no other vertex color is specified
+	*/
+	FColor ConstantVertexColor = FColor::White;
+
+	/**
+	* If true, vertex colors on the FDynamicMesh will be ignored
+	*/
+	bool bIgnoreVertexColors = false;
+
+	/**
+	* If true, a per-triangle color is used to set vertex colors
+	*/
+	bool bUsePerTriangleColor = false;
+
+	/**
+	 * Per-triangle color function. Only called if bUsePerTriangleColor=true
+	 */
+	TFunction<FColor(const FDynamicMesh3*, int)> PerTriangleColorFunc = nullptr;
+
+	/**
+	 * If true, VertexColorRemappingFunc is called on Vertex Colors provided from Mesh to remap them to a different color
+	 */
+	bool bApplyVertexColorRemapping = false;
+
+	/**
+	 * Vertex color remapping function. Only called if bApplyVertexColorRemapping == true, for mesh vertex colors
+	 */
+	TUniqueFunction<void(FVector4f&)> VertexColorRemappingFunc = nullptr;
+
+	/**
+	 * Color Space Transform/Conversion applied to Vertex Colors provided from Mesh Color Overlay Attribute
+	 * Color Space Conversion is applied after any Vertex Color Remapping.
+	 */
+	EDynamicMeshVertexColorTransformMode ColorSpaceTransformMode = EDynamicMeshVertexColorTransformMode::NoTransform;
+
+	/**
+	* If true, a facet normals are used instead of mesh normals
+	*/
+	bool bUsePerTriangleNormals = false;
+
+	/**
+	 * If true, populate secondary buffers using SecondaryTriFilterFunc
+	 */
+	bool bUseSecondaryTriBuffers = false;
+
+	/**
+	 * Filter predicate for secondary triangle index buffer. Only called if bUseSecondaryTriBuffers=true
+	 */
+	TUniqueFunction<bool(const FDynamicMesh3*, int32)> SecondaryTriFilterFunc = nullptr;
+
+
+	/**
+	 * Initialize rendering buffers from given attribute overlays.
+	 * Creates three vertices per triangle, IE no shared vertices in buffers.
+	 */
+	template<typename TriangleEnumerable>
+	void InitializeBuffersFromOverlays(
+										FMeshRenderBufferSet* RenderBuffers,
+										const FDynamicMesh3* Mesh,
+										int NumTriangles, TriangleEnumerable Enumerable,
+										const FDynamicMeshUVOverlay* UVOverlay,
+										const FDynamicMeshNormalOverlay* NormalOverlay,
+										const FDynamicMeshColorOverlay* ColorOverlay,
+										TFunctionRef<void(int, int, int, const FVector3f&, FVector3f&, FVector3f&)> TangentsFunc,
+										bool bTrackTriangles = false)
+	{
+		TArray<const FDynamicMeshUVOverlay*> UVOverlays;
+		UVOverlays.Add(UVOverlay);
+		InitializeBuffersFromOverlays(RenderBuffers, Mesh, NumTriangles, Enumerable,
+			UVOverlays, NormalOverlay, ColorOverlay, TangentsFunc, bTrackTriangles);
+	}
+
+
+	/**
+	* Initialize rendering buffers from given attribute overlays.
+	* Creates three vertices per triangle, IE no shared vertices in buffers.
+	*/
+	template<typename TriangleEnumerable, typename UVOverlayListAllocator>
+	void InitializeBuffersFromOverlays( FMeshRenderBufferSet* RenderBuffers,
+										const FDynamicMesh3* Mesh,
+										int NumTriangles, 
+										TriangleEnumerable Enumerable,
+										const TArray<const FDynamicMeshUVOverlay*, UVOverlayListAllocator>& UVOverlays,
+										const FDynamicMeshNormalOverlay* NormalOverlay,
+										const FDynamicMeshColorOverlay* ColorOverlay,
+										TFunctionRef<void(int, int, int, const FVector3f&, FVector3f&, FVector3f&)> TangentsFunc,
+										bool bTrackTriangles = false);
+												
+
+
+	/**
+	 * Filter the triangles in a FMeshRenderBufferSet into the SecondaryIndexBuffer.
+	 * Requires that RenderBuffers->Triangles has been initialized.
+	 * @param bDuplicate if set, then primary IndexBuffer is unmodified and SecondaryIndexBuffer contains duplicates. Otherwise triangles are sorted via predicate into either primary or secondary.
+	 */
+	void GEOMETRYFRAMEWORK_API UpdateSecondaryTriangleBuffer( FMeshRenderBufferSet* RenderBuffers,
+										                      const FDynamicMesh3* Mesh,
+										                      bool bDuplicate);
+
+
+	/**
+	 * RecomputeRenderBufferTriangleIndexSets re-sorts the existing set of triangles in a FMeshRenderBufferSet
+	 * into primary and secondary index buffers. Note that UploadIndexBufferUpdate() must be called
+	 * after this function!
+	 */
+	void GEOMETRYFRAMEWORK_API RecomputeRenderBufferTriangleIndexSets( FMeshRenderBufferSet* RenderBuffers,
+												                       const FDynamicMesh3* Mesh);
+
+
+	/**
+	 * Update vertex positions/normals/colors of an existing set of render buffers.
+	 * Assumes that buffers were created with unshared vertices, ie three vertices per triangle, eg by InitializeBuffersFromOverlays()
+	 */
+	template<typename TriangleEnumerable>
+	void UpdateVertexBuffersFromOverlays( FMeshRenderBufferSet* RenderBuffers,
+										  const FDynamicMesh3* Mesh,
+										  int NumTriangles, TriangleEnumerable Enumerable,
+										  const FDynamicMeshNormalOverlay* NormalOverlay,
+										  const FDynamicMeshColorOverlay* ColorOverlay,
+										  TFunctionRef<void(int, int, int, const FVector3f&, FVector3f&, FVector3f&)> TangentsFunc,
+										  bool bUpdatePositions = true,
+										  bool bUpdateNormals = false,
+										  bool bUpdateColors = false);
+
+
+	/**
+	 * Update vertex uvs of an existing set of render buffers.
+	 * Assumes that buffers were created with unshared vertices, ie three vertices per triangle, eg by InitializeBuffersFromOverlays()
+	 */
+	template<typename TriangleEnumerable, typename UVOverlayListAllocator>
+	void UpdateVertexUVBufferFromOverlays( FMeshRenderBufferSet* RenderBuffers,
+										   const FDynamicMesh3* Mesh,
+										   int32 NumTriangles, TriangleEnumerable Enumerable,
+										   const TArray<const FDynamicMeshUVOverlay*, UVOverlayListAllocator>& UVOverlays);
+
+
+	/**
+	* Get the overlay color the FColor,  respecting the ColorSpaceTransformMode utilizing the VertexColorRemappingFunc if requested. 
+	*/ 
+	GEOMETRYFRAMEWORK_API FColor GetOverlayColorAsFColor( const FDynamicMeshColorOverlay* ColorOverlay,
+		                                                  int32 ElementID)
+	{
+		checkSlow(ColorOverlay);
+		FVector4f UseColor = ColorOverlay->GetElement(ElementID);
+
+		if (bApplyVertexColorRemapping)
+		{
+			VertexColorRemappingFunc(UseColor);
+		}
+
+		if (ColorSpaceTransformMode == EDynamicMeshVertexColorTransformMode::SRGBToLinear)
+		{
+			// is there a better way to do this? 
+			FColor QuantizedSRGBColor = ((FLinearColor)UseColor).ToFColor(false);
+			return FLinearColor(QuantizedSRGBColor).ToFColor(false);
+		}
+		else
+		{
+			bool bConvertToSRGB = (ColorSpaceTransformMode == EDynamicMeshVertexColorTransformMode::LinearToSRGB);
+			return ((FLinearColor)UseColor).ToFColor(bConvertToSRGB);
+		}
+	}
+
+
+private:
+	using FIndex2i = UE::Geometry::FIndex2i;
+	using FIndex3i = UE::Geometry::FIndex3i;
+	
+};
+
+
+// Simple tool to initialize a single set of mesh buffers for the entire mesh.  
+// for examples of more complicated conversions, see DynamicMeshSceneProxy 
+class FDynamicMeshComponentToMeshRenderBufferSet
+{
+
+public:
+	FMeshRenderBufferSetConversionUtil MeshRenderBufferSetConverter;
+
+	// note, this conversion may recompute the tangents on the DynamicMeshComponent 
+	// since a dynamic mesh component with "autocalculated" tangents will compute them on first request
+	// @param DynamicMeshComponent - the mesh to be converted
+	// @param MeshRenderBufferSet  - the buffer set to be initialized
+	// @param bUseComponentSettings - if true, the Component settings will override (and update)  MeshRederBufferSetBuilder.{ColorspaceTransformMode, bUsePerTriangleNormals}.
+	GEOMETRYFRAMEWORK_API void Convert(class UDynamicMeshComponent& DynamicMeshComponent, FMeshRenderBufferSet& MeshRenderBufferSet, bool bUseComponentSettings = true);
+
+protected:
+	TUniqueFunction<void(int, int, int, const FVector3f&, FVector3f&, FVector3f&)> MakeTangentsFunc(class UDynamicMeshComponent& DynamicMeshCompnent, bool bSkipAutoCompute = false);
+};
+
+
+
+
+
+
+
+
+// -- template implementations -- // 
+
+
+template<typename TriangleEnumerable, typename UVOverlayListAllocator>
+void FMeshRenderBufferSetConversionUtil::InitializeBuffersFromOverlays( FMeshRenderBufferSet* RenderBuffers,
+																	    const FDynamicMesh3* Mesh,
+																		int NumTriangles,
+																		TriangleEnumerable Enumerable,
+																	    const TArray<const FDynamicMeshUVOverlay*, UVOverlayListAllocator>& UVOverlays,
+																		const FDynamicMeshNormalOverlay* NormalOverlay,
+																		const FDynamicMeshColorOverlay* ColorOverlay,
+																		TFunctionRef<void(int, int, int, const FVector3f&, FVector3f&, FVector3f&)> TangentsFunc,
+																		bool bTrackTriangles)
+
+{
+	RenderBuffers->TriangleCount = NumTriangles;
+	if (NumTriangles == 0)
+	{
+		return;
+	}
+
+	bool bHaveColors = (ColorOverlay != nullptr) && (bIgnoreVertexColors == false);
+
+	int NumVertices = NumTriangles * 3;
+	int NumUVOverlays = UVOverlays.Num();
+	int NumTexCoords = FMath::Max(1, NumUVOverlays);		// must have at least one tex coord
+	TArray<FIndex3i, TFixedAllocator<MAX_STATIC_TEXCOORDS>> UVTriangles;
+	UVTriangles.SetNum(NumTexCoords);
+
+	{
+		RenderBuffers->PositionVertexBuffer.Init(NumVertices);
+		RenderBuffers->StaticMeshVertexBuffer.Init(NumVertices, NumTexCoords);
+		RenderBuffers->ColorVertexBuffer.Init(NumVertices);
+		RenderBuffers->IndexBuffer.Indices.AddUninitialized(NumTriangles * 3);
+	}
+
+	// build triangle list if requested, or if we are using secondary buffers in which case we need it to filter later
+	bool bBuildTriangleList = bTrackTriangles || bUseSecondaryTriBuffers;
+	if (bBuildTriangleList)
+	{
+		RenderBuffers->Triangles = TArray<int32>();
+	}
+
+	int TriIdx = 0, VertIdx = 0;
+	FVector3f TangentX, TangentY;
+	for (int TriangleID : Enumerable)
+	{
+		FIndex3i Tri = Mesh->GetTriangle(TriangleID);
+		for (int32 k = 0; k < NumTexCoords; ++k)
+		{
+			UVTriangles[k] = (k < NumUVOverlays && UVOverlays[k] != nullptr) ? UVOverlays[k]->GetTriangle(TriangleID) : FIndex3i::Invalid();
+		}
+		FIndex3i TriNormal = (NormalOverlay != nullptr) ? NormalOverlay->GetTriangle(TriangleID) : FIndex3i::Zero();
+		FIndex3i TriColor = (ColorOverlay != nullptr) ? ColorOverlay->GetTriangle(TriangleID) : FIndex3i::Zero();
+
+		FColor UniformTriColor = ConstantVertexColor;
+		if (bUsePerTriangleColor && PerTriangleColorFunc != nullptr)
+		{
+			UniformTriColor = PerTriangleColorFunc(Mesh, TriangleID);
+			bHaveColors = false;
+		}
+
+		for (int j = 0; j < 3; ++j)
+		{
+			RenderBuffers->PositionVertexBuffer.VertexPosition(VertIdx) = (FVector3f)Mesh->GetVertex(Tri[j]);
+
+			FVector3f Normal;
+			if (bUsePerTriangleNormals)
+			{
+				Normal = (FVector3f)Mesh->GetTriNormal(TriangleID);
+			}
+			else
+			{
+				Normal = (NormalOverlay != nullptr && TriNormal[j] != FDynamicMesh3::InvalidID) ?
+					NormalOverlay->GetElement(TriNormal[j]) : Mesh->GetVertexNormal(Tri[j]);
+			}
+
+			// get tangents
+			TangentsFunc(Tri[j], TriangleID, j, Normal, TangentX, TangentY);
+
+			RenderBuffers->StaticMeshVertexBuffer.SetVertexTangents(VertIdx, TangentX, TangentY, Normal);
+
+			for (int32 k = 0; k < NumTexCoords; ++k)
+			{
+				FVector2f UV = (UVTriangles[k][j] != FDynamicMesh3::InvalidID) ?
+					UVOverlays[k]->GetElement(UVTriangles[k][j]) : FVector2f::Zero();
+				RenderBuffers->StaticMeshVertexBuffer.SetVertexUV(VertIdx, k, UV);
+			}
+
+			FColor VertexFColor = (bHaveColors && TriColor[j] != FDynamicMesh3::InvalidID) ?
+				GetOverlayColorAsFColor(ColorOverlay, TriColor[j]) : UniformTriColor;
+
+			RenderBuffers->ColorVertexBuffer.VertexColor(VertIdx) = VertexFColor;
+
+			RenderBuffers->IndexBuffer.Indices[TriIdx++] = VertIdx;		// currently TriIdx == VertIdx so we don't really need both...
+			VertIdx++;
+		}
+
+		if (bBuildTriangleList)
+		{
+			RenderBuffers->Triangles->Add(TriangleID);
+		}
+	}
+
+	// split triangles into secondary buffer (at bit redundant since we just built IndexBuffer, but we may optionally duplicate triangles in the future)
+	if (bUseSecondaryTriBuffers)
+	{
+		RenderBuffers->bEnableSecondaryIndexBuffer = true;
+		UpdateSecondaryTriangleBuffer(RenderBuffers, Mesh, false);
+	}
+}
+
+
+/**
+* Update vertex positions/normals/colors of an existing set of render buffers.
+* Assumes that buffers were created with unshared vertices, ie three vertices per triangle, eg by InitializeBuffersFromOverlays()
+*/
+template<typename TriangleEnumerable>
+void FMeshRenderBufferSetConversionUtil::UpdateVertexBuffersFromOverlays(  FMeshRenderBufferSet* RenderBuffers,
+																		   const FDynamicMesh3* Mesh,
+																		   int NumTriangles, TriangleEnumerable Enumerable,
+																		   const FDynamicMeshNormalOverlay* NormalOverlay,
+																		   const FDynamicMeshColorOverlay* ColorOverlay,
+																		   TFunctionRef<void(int, int, int, const FVector3f&, FVector3f&, FVector3f&)> TangentsFunc,
+																		   bool bUpdatePositions,
+																		   bool bUpdateNormals,
+																		   bool bUpdateColors)
+{
+	if (RenderBuffers->TriangleCount == 0)
+	{
+		return;
+	}
+
+	bool bHaveColors = (ColorOverlay != nullptr) && (bIgnoreVertexColors == false);
+
+	int NumVertices = NumTriangles * 3;
+	if ((bUpdatePositions && ensure(RenderBuffers->PositionVertexBuffer.GetNumVertices() == NumVertices) == false)
+		|| (bUpdateNormals && ensure(RenderBuffers->StaticMeshVertexBuffer.GetNumVertices() == NumVertices) == false)
+		|| (bUpdateColors && ensure(RenderBuffers->ColorVertexBuffer.GetNumVertices() == NumVertices) == false))
+	{
+		return;
+	}
+
+	int VertIdx = 0;
+	FVector3f TangentX, TangentY;
+	for (int TriangleID : Enumerable)
+	{
+		FIndex3i Tri = Mesh->GetTriangle(TriangleID);
+
+		FIndex3i TriNormal = (bUpdateNormals && NormalOverlay != nullptr) ? NormalOverlay->GetTriangle(TriangleID) : FIndex3i::Zero();
+		FIndex3i TriColor = (bUpdateColors && ColorOverlay != nullptr) ? ColorOverlay->GetTriangle(TriangleID) : FIndex3i::Zero();
+
+		FColor UniformTriColor = ConstantVertexColor;
+		if (bUpdateColors && bUsePerTriangleColor && PerTriangleColorFunc != nullptr)
+		{
+			UniformTriColor = PerTriangleColorFunc(Mesh, TriangleID);
+			bHaveColors = false;
+		}
+
+		for (int j = 0; j < 3; ++j)
+		{
+			if (bUpdatePositions)
+			{
+				RenderBuffers->PositionVertexBuffer.VertexPosition(VertIdx) = (FVector3f)Mesh->GetVertex(Tri[j]);
+			}
+
+			if (bUpdateNormals)
+			{
+				// get normal and tangent
+				FVector3f Normal;
+				if (bUsePerTriangleNormals)
+				{
+					Normal = (FVector3f)Mesh->GetTriNormal(TriangleID);
+				}
+				else
+				{
+					Normal = (NormalOverlay != nullptr && TriNormal[j] != FDynamicMesh3::InvalidID) ?
+						NormalOverlay->GetElement(TriNormal[j]) : Mesh->GetVertexNormal(Tri[j]);
+				}
+
+				TangentsFunc(Tri[j], TriangleID, j, Normal, TangentX, TangentY);
+
+				RenderBuffers->StaticMeshVertexBuffer.SetVertexTangents(VertIdx, (FVector3f)TangentX, (FVector3f)TangentY, (FVector3f)Normal);
+			}
+
+			if (bUpdateColors)
+			{
+				FColor VertexFColor = (bHaveColors && TriColor[j] != FDynamicMesh3::InvalidID) ?
+					GetOverlayColorAsFColor(ColorOverlay, TriColor[j]) : UniformTriColor;
+				RenderBuffers->ColorVertexBuffer.VertexColor(VertIdx) = VertexFColor;
+			}
+
+			VertIdx++;
+		}
+	}
+}
+
+template<typename TriangleEnumerable, typename UVOverlayListAllocator>
+void FMeshRenderBufferSetConversionUtil::UpdateVertexUVBufferFromOverlays( FMeshRenderBufferSet* RenderBuffers,
+																	       const FDynamicMesh3* Mesh,
+																	       int32 NumTriangles, TriangleEnumerable Enumerable,
+																	       const TArray<const FDynamicMeshUVOverlay*, UVOverlayListAllocator>& UVOverlays)
+{
+	// We align the update to the way we set UV's in InitializeBuffersFromOverlays.
+
+	if (RenderBuffers->TriangleCount == 0)
+	{
+		return;
+	}
+	int NumVertices = NumTriangles * 3;
+	if (ensure(RenderBuffers->StaticMeshVertexBuffer.GetNumVertices() == NumVertices) == false)
+	{
+		return;
+	}
+
+	int NumUVOverlays = UVOverlays.Num();
+	int NumTexCoords = RenderBuffers->StaticMeshVertexBuffer.GetNumTexCoords();
+	if (!ensure(NumUVOverlays <= NumTexCoords))
+	{
+		return;
+	}
+
+	// Temporarily stores the UV element indices for all UV channels of a single triangle
+	TArray<FIndex3i, TFixedAllocator<MAX_STATIC_TEXCOORDS>> UVTriangles;
+	UVTriangles.SetNum(NumTexCoords);
+
+	int VertIdx = 0;
+	for (int TriangleID : Enumerable)
+	{
+		for (int32 k = 0; k < NumTexCoords; ++k)
+		{
+			UVTriangles[k] = (k < NumUVOverlays && UVOverlays[k] != nullptr) ? UVOverlays[k]->GetTriangle(TriangleID) : FIndex3i::Invalid();
+		}
+
+		for (int j = 0; j < 3; ++j)
+		{
+			for (int32 k = 0; k < NumTexCoords; ++k)
+			{
+				FVector2f UV = (UVTriangles[k][j] != FDynamicMesh3::InvalidID) ?
+					UVOverlays[k]->GetElement(UVTriangles[k][j]) : FVector2f::Zero();
+				RenderBuffers->StaticMeshVertexBuffer.SetVertexUV(VertIdx, k, UV);
+			}
+
+			++VertIdx;
+		}
+	}
+}
+
