Add option to remove extra material slots when configuring materials on a dynamic mesh component from an array, so that it will only have the materials in the array + make that the default behavior

diff --git a/Engine/Source/Runtime/GeometryFramework/Private/Components/DynamicMeshComponent.cpp b/Engine/Source/Runtime/GeometryFramework/Private/Components/DynamicMeshComponent.cpp
--- a/Engine/Source/Runtime/GeometryFramework/Private/Components/DynamicMeshComponent.cpp
+++ b/Engine/Source/Runtime/GeometryFramework/Private/Components/DynamicMeshComponent.cpp
@@ -316,12 +316,16 @@ bool UDynamicMeshComponent::ValidateMaterialSlots(bool bCreateIfMissing, bool bD
 }
 
 
-void UDynamicMeshComponent::ConfigureMaterialSet(const TArray<UMaterialInterface*>& NewMaterialSet)
+void UDynamicMeshComponent::ConfigureMaterialSet(const TArray<UMaterialInterface*>& NewMaterialSet, bool bDeleteExtraSlots)
 {
 	for (int k = 0; k < NewMaterialSet.Num(); ++k)
 	{
 		SetMaterial(k, NewMaterialSet[k]);
 	}	
+	if (bDeleteExtraSlots)
+	{
+		SetNumMaterials(NewMaterialSet.Num());
+	}
 }
 
 
diff --git a/Engine/Source/Runtime/GeometryFramework/Public/Components/DynamicMeshComponent.h b/Engine/Source/Runtime/GeometryFramework/Public/Components/DynamicMeshComponent.h
--- a/Engine/Source/Runtime/GeometryFramework/Public/Components/DynamicMeshComponent.h
+++ b/Engine/Source/Runtime/GeometryFramework/Public/Components/DynamicMeshComponent.h
@@ -571,9 +571,10 @@ public:
 	 * Set new list of Materials for the Mesh. Dynamic Mesh Component does not have 
 	 * Slot Names, so the size of the Material Set should be the same as the number of
 	 * different Material IDs on the mesh MaterialID attribute
+	 * @param bDeleteExtraSlots if true, extra Material Slots beyond max NewMaterialSet.Num() are removed
 	 */
 	UFUNCTION(BlueprintCallable, Category = "Dynamic Mesh Component")
-	GEOMETRYFRAMEWORK_API void ConfigureMaterialSet(const TArray<UMaterialInterface*>& NewMaterialSet);
+	GEOMETRYFRAMEWORK_API void ConfigureMaterialSet(const TArray<UMaterialInterface*>& NewMaterialSet, bool bDeleteExtraSlots = true);
 
 	/**
 	 * Compute the maximum MaterialID on the DynamicMesh, and ensure that Material Slots match.
