Consistently use FPrimitiveUniformShaderParametersBuilder to set properties on dynamic primitives

This is to catch all variables that can impact shading (such as holdouts). This extends the fix that was done for GeometryCache actors.

diff --git a/Engine/Plugins/Experimental/ImagePlate/Source/ImagePlate/Private/ImagePlateComponent.cpp b/Engine/Plugins/Experimental/ImagePlate/Source/ImagePlate/Private/ImagePlateComponent.cpp
--- a/Engine/Plugins/Experimental/ImagePlate/Source/ImagePlate/Private/ImagePlateComponent.cpp
+++ b/Engine/Plugins/Experimental/ImagePlate/Source/ImagePlate/Private/ImagePlateComponent.cpp
@@ -9,6 +9,7 @@
 #include "Engine/CollisionProfile.h"
 #include "RayTracingInstance.h"
 #include "SceneManagement.h"
+#include "PrimitiveUniformShaderParametersBuilder.h"
 #include "Engine/Engine.h"
 #include "Materials/MaterialInstanceDynamic.h"
 #include "DynamicMeshBuilder.h"
@@ -283,16 +284,11 @@ namespace
 				FMeshBatchElement& BatchElement = MeshBatch.Elements[0];
 				BatchElement.IndexBuffer = &IndexBuffer;
 
-				bool bHasPrecomputedVolumetricLightmap = false;
-				FMatrix PreviousLocalToWorld = {};
-				PreviousLocalToWorld.SetIdentity();
-				int32 SingleCaptureIndex = 0;
-				bool bOutputVelocity = false;
-				GetScene().GetPrimitiveUniformShaderParameters_RenderThread(GetPrimitiveSceneInfo(), bHasPrecomputedVolumetricLightmap, PreviousLocalToWorld, SingleCaptureIndex, bOutputVelocity);
-				bOutputVelocity |= AlwaysHasVelocity();
-
 				FDynamicPrimitiveUniformBuffer& DynamicPrimitiveUniformBuffer = Context.RayTracingMeshResourceCollector.AllocateOneFrameResource<FDynamicPrimitiveUniformBuffer>();
-				DynamicPrimitiveUniformBuffer.Set(Context.RHICmdList, GetLocalToWorld(), PreviousLocalToWorld, GetBounds(), GetLocalBounds(), GetLocalBounds(), true, bHasPrecomputedVolumetricLightmap, bOutputVelocity, GetCustomPrimitiveData());
+				FPrimitiveUniformShaderParametersBuilder Builder;
+				BuildUniformShaderParameters(Builder);
+				DynamicPrimitiveUniformBuffer.Set(Context.RayTracingMeshResourceCollector.GetRHICommandList(), Builder);
+
 				BatchElement.PrimitiveUniformBufferResource = &DynamicPrimitiveUniformBuffer.UniformBuffer;
 
 				BatchElement.FirstIndex = 0;
diff --git a/Engine/Plugins/Runtime/CableComponent/Source/CableComponent/Private/CableComponent.cpp b/Engine/Plugins/Runtime/CableComponent/Source/CableComponent/Private/CableComponent.cpp
--- a/Engine/Plugins/Runtime/CableComponent/Source/CableComponent/Private/CableComponent.cpp
+++ b/Engine/Plugins/Runtime/CableComponent/Source/CableComponent/Private/CableComponent.cpp
@@ -5,6 +5,7 @@
 #include "PrimitiveSceneProxy.h"
 #include "MaterialDomain.h"
 #include "SceneManagement.h"
+#include "PrimitiveUniformShaderParametersBuilder.h"
 #include "Engine/CollisionProfile.h"
 #include "Materials/Material.h"
 #include "Materials/MaterialRenderProxy.h"
@@ -364,15 +365,11 @@ public:
 				Mesh.VertexFactory = &VertexFactory;
 				Mesh.MaterialRenderProxy = MaterialProxy;
 
-				bool bHasPrecomputedVolumetricLightmap;
-				FMatrix PreviousLocalToWorld;
-				int32 SingleCaptureIndex;
-				bool bOutputVelocity;
-				GetScene().GetPrimitiveUniformShaderParameters_RenderThread(GetPrimitiveSceneInfo(), bHasPrecomputedVolumetricLightmap, PreviousLocalToWorld, SingleCaptureIndex, bOutputVelocity);
-				bOutputVelocity |= AlwaysHasVelocity();
-
 				FDynamicPrimitiveUniformBuffer& DynamicPrimitiveUniformBuffer = Collector.AllocateOneFrameResource<FDynamicPrimitiveUniformBuffer>();
-				DynamicPrimitiveUniformBuffer.Set(Collector.GetRHICommandList(), GetLocalToWorld(), PreviousLocalToWorld, GetBounds(), GetLocalBounds(), GetLocalBounds(), true, bHasPrecomputedVolumetricLightmap, bOutputVelocity, GetCustomPrimitiveData());
+				FPrimitiveUniformShaderParametersBuilder Builder;
+				BuildUniformShaderParameters(Builder);
+				DynamicPrimitiveUniformBuffer.Set(Collector.GetRHICommandList(), Builder);
+
 				BatchElement.PrimitiveUniformBufferResource = &DynamicPrimitiveUniformBuffer.UniformBuffer;
 
 				BatchElement.FirstIndex = 0;
@@ -512,15 +509,11 @@ public:
 			FMeshBatchElement& BatchElement = MeshBatch.Elements[0];
 			BatchElement.IndexBuffer = &IndexBuffer;
 
-			bool bHasPrecomputedVolumetricLightmap;
-			FMatrix PreviousLocalToWorld;
-			int32 SingleCaptureIndex;
-			bool bOutputVelocity;
-			GetScene().GetPrimitiveUniformShaderParameters_RenderThread(GetPrimitiveSceneInfo(), bHasPrecomputedVolumetricLightmap, PreviousLocalToWorld, SingleCaptureIndex, bOutputVelocity);
-			bOutputVelocity |= AlwaysHasVelocity();
-
 			FDynamicPrimitiveUniformBuffer& DynamicPrimitiveUniformBuffer = Context.RayTracingMeshResourceCollector.AllocateOneFrameResource<FDynamicPrimitiveUniformBuffer>();
-			DynamicPrimitiveUniformBuffer.Set(Context.RHICmdList, GetLocalToWorld(), PreviousLocalToWorld, GetBounds(), GetLocalBounds(), GetLocalBounds(), ReceivesDecals(), bHasPrecomputedVolumetricLightmap, bOutputVelocity, GetCustomPrimitiveData());
+			FPrimitiveUniformShaderParametersBuilder Builder;
+			BuildUniformShaderParameters(Builder);
+			DynamicPrimitiveUniformBuffer.Set(Context.RayTracingMeshResourceCollector.GetRHICommandList(), Builder);
+
 			BatchElement.PrimitiveUniformBufferResource = &DynamicPrimitiveUniformBuffer.UniformBuffer;
 			BatchElement.FirstIndex = 0;
 			BatchElement.NumPrimitives = GetRequiredIndexCount() / 3;
diff --git a/Engine/Plugins/Runtime/CustomMeshComponent/Source/CustomMeshComponent/Private/CustomMeshComponent.cpp b/Engine/Plugins/Runtime/CustomMeshComponent/Source/CustomMeshComponent/Private/CustomMeshComponent.cpp
--- a/Engine/Plugins/Runtime/CustomMeshComponent/Source/CustomMeshComponent/Private/CustomMeshComponent.cpp
+++ b/Engine/Plugins/Runtime/CustomMeshComponent/Source/CustomMeshComponent/Private/CustomMeshComponent.cpp
@@ -8,6 +8,7 @@
 #include "Engine/CollisionProfile.h"
 #include "Materials/Material.h"
 #include "SceneManagement.h"
+#include "PrimitiveUniformShaderParametersBuilder.h"
 #include "DynamicMeshBuilder.h"
 #include "Engine/Engine.h"
 #include "SceneInterface.h"
@@ -126,15 +127,10 @@ public:
 				Mesh.VertexFactory = &VertexFactory;
 				Mesh.MaterialRenderProxy = MaterialProxy;
 
-				bool bHasPrecomputedVolumetricLightmap;
-				FMatrix PreviousLocalToWorld;
-				int32 SingleCaptureIndex;
-				bool bOutputVelocity;
-				GetScene().GetPrimitiveUniformShaderParameters_RenderThread(GetPrimitiveSceneInfo(), bHasPrecomputedVolumetricLightmap, PreviousLocalToWorld, SingleCaptureIndex, bOutputVelocity);
-				bOutputVelocity |= AlwaysHasVelocity();
-
 				FDynamicPrimitiveUniformBuffer& DynamicPrimitiveUniformBuffer = Collector.AllocateOneFrameResource<FDynamicPrimitiveUniformBuffer>();
-				DynamicPrimitiveUniformBuffer.Set(Collector.GetRHICommandList(), GetLocalToWorld(), PreviousLocalToWorld, GetBounds(), GetLocalBounds(), true, bHasPrecomputedVolumetricLightmap, bOutputVelocity);
+				FPrimitiveUniformShaderParametersBuilder Builder;
+				BuildUniformShaderParameters(Builder);
+				DynamicPrimitiveUniformBuffer.Set(Collector.GetRHICommandList(), Builder);
 				BatchElement.PrimitiveUniformBufferResource = &DynamicPrimitiveUniformBuffer.UniformBuffer;
 
 				BatchElement.FirstIndex = 0;
diff --git a/Engine/Plugins/Runtime/ProceduralMeshComponent/Source/ProceduralMeshComponent/Private/ProceduralMeshComponent.cpp b/Engine/Plugins/Runtime/ProceduralMeshComponent/Source/ProceduralMeshComponent/Private/ProceduralMeshComponent.cpp
--- a/Engine/Plugins/Runtime/ProceduralMeshComponent/Source/ProceduralMeshComponent/Private/ProceduralMeshComponent.cpp
+++ b/Engine/Plugins/Runtime/ProceduralMeshComponent/Source/ProceduralMeshComponent/Private/ProceduralMeshComponent.cpp
@@ -10,6 +10,7 @@
 #include "Engine/Engine.h"
 #include "RenderUtils.h"
 #include "SceneManagement.h"
+#include "PrimitiveUniformShaderParametersBuilder.h"
 #include "PhysicsEngine/BodySetup.h"
 #include "ProceduralMeshComponentPluginPrivate.h"
 #include "DynamicMeshBuilder.h"
@@ -366,15 +367,11 @@ public:
 						Mesh.VertexFactory = &Section->VertexFactory;
 						Mesh.MaterialRenderProxy = MaterialProxy;
 
-						bool bHasPrecomputedVolumetricLightmap;
-						FMatrix PreviousLocalToWorld;
-						int32 SingleCaptureIndex;
-						bool bOutputVelocity;
-						GetScene().GetPrimitiveUniformShaderParameters_RenderThread(GetPrimitiveSceneInfo(), bHasPrecomputedVolumetricLightmap, PreviousLocalToWorld, SingleCaptureIndex, bOutputVelocity);
-						bOutputVelocity |= AlwaysHasVelocity();
-
 						FDynamicPrimitiveUniformBuffer& DynamicPrimitiveUniformBuffer = Collector.AllocateOneFrameResource<FDynamicPrimitiveUniformBuffer>();
-						DynamicPrimitiveUniformBuffer.Set(Collector.GetRHICommandList(), GetLocalToWorld(), PreviousLocalToWorld, GetBounds(), GetLocalBounds(), GetLocalBounds(), ReceivesDecals(), bHasPrecomputedVolumetricLightmap, bOutputVelocity, GetCustomPrimitiveData());
+						FPrimitiveUniformShaderParametersBuilder Builder;
+						BuildUniformShaderParameters(Builder);
+						DynamicPrimitiveUniformBuffer.Set(Collector.GetRHICommandList(), Builder);
+
 						BatchElement.PrimitiveUniformBufferResource = &DynamicPrimitiveUniformBuffer.UniformBuffer;
 
 						BatchElement.FirstIndex = 0;
@@ -484,15 +481,11 @@ public:
 					FMeshBatchElement& BatchElement = MeshBatch.Elements[0];
 					BatchElement.IndexBuffer = &Section->IndexBuffer;
 
-					bool bHasPrecomputedVolumetricLightmap;
-					FMatrix PreviousLocalToWorld;
-					int32 SingleCaptureIndex;
-					bool bOutputVelocity;
-					GetScene().GetPrimitiveUniformShaderParameters_RenderThread(GetPrimitiveSceneInfo(), bHasPrecomputedVolumetricLightmap, PreviousLocalToWorld, SingleCaptureIndex, bOutputVelocity);
-					bOutputVelocity |= AlwaysHasVelocity();
-
 					FDynamicPrimitiveUniformBuffer& DynamicPrimitiveUniformBuffer = Context.RayTracingMeshResourceCollector.AllocateOneFrameResource<FDynamicPrimitiveUniformBuffer>();
-					DynamicPrimitiveUniformBuffer.Set(Context.RHICmdList, GetLocalToWorld(), PreviousLocalToWorld, GetBounds(), GetLocalBounds(), GetLocalBounds(), ReceivesDecals(), bHasPrecomputedVolumetricLightmap, bOutputVelocity, GetCustomPrimitiveData());
+					FPrimitiveUniformShaderParametersBuilder Builder;
+					BuildUniformShaderParameters(Builder);
+					DynamicPrimitiveUniformBuffer.Set(Context.RayTracingMeshResourceCollector.GetRHICommandList(), Builder);
+
 					BatchElement.PrimitiveUniformBufferResource = &DynamicPrimitiveUniformBuffer.UniformBuffer;
 
 					BatchElement.FirstIndex = 0;
diff --git a/Engine/Source/Runtime/GeometryFramework/Private/Components/BaseDynamicMeshSceneProxy.cpp b/Engine/Source/Runtime/GeometryFramework/Private/Components/BaseDynamicMeshSceneProxy.cpp
--- a/Engine/Source/Runtime/GeometryFramework/Private/Components/BaseDynamicMeshSceneProxy.cpp
+++ b/Engine/Source/Runtime/GeometryFramework/Private/Components/BaseDynamicMeshSceneProxy.cpp
@@ -8,6 +8,7 @@
 #include "RayTracingDefinitions.h"
 #include "RayTracingInstance.h"
 #include "SceneInterface.h"
+#include "PrimitiveUniformShaderParametersBuilder.h"
 #include "SceneManagement.h"
 #include "Engine/Engine.h"		// for GEngine definition
 #include "MeshCardRepresentation.h"
@@ -297,13 +298,6 @@ void FBaseDynamicMeshSceneProxy::GetDynamicMeshElements(const TArray<const FScen
 		{
 			const FSceneView* View = Views[ViewIndex];
 
-			bool bHasPrecomputedVolumetricLightmap;
-			FMatrix PreviousLocalToWorld;
-			int32 SingleCaptureIndex;
-			bool bOutputVelocity;
-			GetScene().GetPrimitiveUniformShaderParameters_RenderThread(GetPrimitiveSceneInfo(), bHasPrecomputedVolumetricLightmap, PreviousLocalToWorld, SingleCaptureIndex, bOutputVelocity);
-			bOutputVelocity |= AlwaysHasVelocity();
-
 			// Draw the mesh.
 			for (FMeshRenderBufferSet* BufferSet : Buffers)
 			{
@@ -328,7 +322,9 @@ void FBaseDynamicMeshSceneProxy::GetDynamicMeshElements(const TArray<const FScen
 
 				// do we need separate one of these for each MeshRenderBufferSet?
 				FDynamicPrimitiveUniformBuffer& DynamicPrimitiveUniformBuffer = Collector.AllocateOneFrameResource<FDynamicPrimitiveUniformBuffer>();
-				DynamicPrimitiveUniformBuffer.Set(Collector.GetRHICommandList(), GetLocalToWorld(), PreviousLocalToWorld, GetActorPosition(), GetBounds(), GetLocalBounds(), GetLocalBounds(), true, bHasPrecomputedVolumetricLightmap, bOutputVelocity, GetCustomPrimitiveData());
+				FPrimitiveUniformShaderParametersBuilder Builder;
+				BuildUniformShaderParameters(Builder);
+				DynamicPrimitiveUniformBuffer.Set(Collector.GetRHICommandList(), Builder);
 
 				// If we want Wireframe-on-Shaded, we have to draw the solid. If View Mode Overrides are enabled, the solid
 				// will be replaced with it's wireframe, so we might as well not. 
@@ -425,13 +421,6 @@ void FBaseDynamicMeshSceneProxy::GetCollisionDynamicMeshElements(TArray<FMeshRen
 					FColoredMaterialRenderProxy* CollisionMaterialInstance = new FColoredMaterialRenderProxy(MaterialToUse->GetRenderProxy(), DrawCollisionColor);
 					Collector.RegisterOneFrameMaterialProxy(CollisionMaterialInstance);
 
-					bool bHasPrecomputedVolumetricLightmap;
-					FMatrix PreviousLocalToWorld;
-					int32 SingleCaptureIndex;
-					bool bOutputVelocity;
-					GetScene().GetPrimitiveUniformShaderParameters_RenderThread(GetPrimitiveSceneInfo(), bHasPrecomputedVolumetricLightmap, PreviousLocalToWorld, SingleCaptureIndex, bOutputVelocity);
-					bOutputVelocity |= AlwaysHasVelocity();
-
 					// Draw the mesh with collision materials
 					for (FMeshRenderBufferSet* BufferSet : Buffers)
 					{
@@ -446,7 +435,9 @@ void FBaseDynamicMeshSceneProxy::GetCollisionDynamicMeshElements(TArray<FMeshRen
 
 						// do we need separate one of these for each MeshRenderBufferSet?
 						FDynamicPrimitiveUniformBuffer& DynamicPrimitiveUniformBuffer = Collector.AllocateOneFrameResource<FDynamicPrimitiveUniformBuffer>();
-						DynamicPrimitiveUniformBuffer.Set(Collector.GetRHICommandList(), GetLocalToWorld(), PreviousLocalToWorld, GetBounds(), GetLocalBounds(), GetLocalBounds(), true, bHasPrecomputedVolumetricLightmap, bOutputVelocity, GetCustomPrimitiveData());
+						FPrimitiveUniformShaderParametersBuilder Builder;
+						BuildUniformShaderParameters(Builder);
+						DynamicPrimitiveUniformBuffer.Set(Collector.GetRHICommandList(), Builder);
 
 						if (BufferSet->IndexBuffer.Indices.Num() > 0)
 						{
@@ -690,16 +681,11 @@ void FBaseDynamicMeshSceneProxy::GetDynamicRayTracingInstances(FRayTracingMateri
 	}
 	bool bDrawSecondaryBuffers = ParentBaseComponent->GetSecondaryBuffersVisibility();
 
-	bool bHasPrecomputedVolumetricLightmap;
-	FMatrix PreviousLocalToWorld;
-	int32 SingleCaptureIndex;
-	bool bOutputVelocity;
-	GetScene().GetPrimitiveUniformShaderParameters_RenderThread(GetPrimitiveSceneInfo(), bHasPrecomputedVolumetricLightmap, PreviousLocalToWorld, SingleCaptureIndex, bOutputVelocity);
-	bOutputVelocity |= AlwaysHasVelocity();
-
 	// is it safe to share this between primary and secondary raytracing batches?
 	FDynamicPrimitiveUniformBuffer& DynamicPrimitiveUniformBuffer = Context.RayTracingMeshResourceCollector.AllocateOneFrameResource<FDynamicPrimitiveUniformBuffer>();
-	DynamicPrimitiveUniformBuffer.Set(Context.RHICmdList, GetLocalToWorld(), PreviousLocalToWorld, GetBounds(), GetLocalBounds(), true, bHasPrecomputedVolumetricLightmap, bOutputVelocity);
+	FPrimitiveUniformShaderParametersBuilder Builder;
+	BuildUniformShaderParameters(Builder);
+	DynamicPrimitiveUniformBuffer.Set(Context.RayTracingMeshResourceCollector.GetRHICommandList(), Builder);
 
 	// Draw the active buffer sets
 	for (FMeshRenderBufferSet* BufferSet : Buffers)
