add IsElementSelectable concept to dynamic mesh components, so we can disable element selection for generated dynamic meshes
Expose SetIsElementSelectable method to BP, so it can be re-enabled if desired

diff --git a/Engine/Plugins/Editor/ModelingToolsEditorMode/Source/ModelingToolsEditorMode/Private/ModelingToolsEditorMode.cpp b/Engine/Plugins/Editor/ModelingToolsEditorMode/Source/ModelingToolsEditorMode/Private/ModelingToolsEditorMode.cpp
--- a/Engine/Plugins/Editor/ModelingToolsEditorMode/Source/ModelingToolsEditorMode/Private/ModelingToolsEditorMode.cpp
+++ b/Engine/Plugins/Editor/ModelingToolsEditorMode/Source/ModelingToolsEditorMode/Private/ModelingToolsEditorMode.cpp
@@ -24,6 +24,7 @@
 #include "Snapping/ModelingSceneSnappingManager.h"
 #include "Scene/LevelObjectsObserver.h"
 #include "UnrealEdGlobals.h" // GUnrealEd
+#include "Algo/RemoveIf.h"
 
 #include "Features/IModularFeatures.h"
 #include "ModelingModeToolExtensions.h"
@@ -1508,6 +1509,9 @@ void UModelingToolsEditorMode::UpdateSelectionManagerOnEditorSelectionChange(boo
 		}
 	}
 
+	// filter out any dynamic mesh components that aren't editable or aren't element-selectable
+	SelectedDynamicMeshComponents.SetNum(Algo::RemoveIf(SelectedDynamicMeshComponents, [](UDynamicMeshComponent* DMC)->bool {return !DMC->IsEditable() || !DMC->IsElementSelectable();}));
+
 	// convert selected Component types into selection Identifiers
 	TArray<FGeometryIdentifier> ValidIdentifiers;
 	for (UDynamicMeshComponent* DynamicMeshComponent : SelectedDynamicMeshComponents)
diff --git a/Engine/Plugins/Runtime/GeometryScripting/Source/GeometryScriptingEditor/Private/GeometryActors/GeneratedDynamicMeshActor.cpp b/Engine/Plugins/Runtime/GeometryScripting/Source/GeometryScriptingEditor/Private/GeometryActors/GeneratedDynamicMeshActor.cpp
--- a/Engine/Plugins/Runtime/GeometryScripting/Source/GeometryScriptingEditor/Private/GeometryActors/GeneratedDynamicMeshActor.cpp
+++ b/Engine/Plugins/Runtime/GeometryScripting/Source/GeometryScriptingEditor/Private/GeometryActors/GeneratedDynamicMeshActor.cpp
@@ -47,6 +47,12 @@ void AGeneratedDynamicMeshActor::PostActorCreated()
 {
 	Super::PostActorCreated();
 	RegisterWithGenerationManager();
+	
+	// Default newly-created generated dynamic meshes to not be element-selectable
+	if (UDynamicMeshComponent* Component = GetDynamicMeshComponent())
+	{
+		Component->SetIsElementSelectable(false);
+	}
 }
 
 void AGeneratedDynamicMeshActor::Destroyed()
diff --git a/Engine/Source/Runtime/GeometryFramework/Public/Components/DynamicMeshComponent.h b/Engine/Source/Runtime/GeometryFramework/Public/Components/DynamicMeshComponent.h
--- a/Engine/Source/Runtime/GeometryFramework/Public/Components/DynamicMeshComponent.h
+++ b/Engine/Source/Runtime/GeometryFramework/Public/Components/DynamicMeshComponent.h
@@ -233,6 +233,20 @@ public:
 	GEOMETRYFRAMEWORK_API bool IsEditable() const { return bIsEditable; }
 	GEOMETRYFRAMEWORK_API void SetIsEditable(bool bInIsEditable) { bIsEditable = bInIsEditable; }
 
+
+	/**
+	 * @return Whether elements of the dynamic mesh can be interactively selected.
+	 */
+	UFUNCTION(BlueprintCallable, Category = "Dynamic Mesh Component")
+	GEOMETRYFRAMEWORK_API bool IsElementSelectable() const { return bIsElementSelectable; }
+	/**
+	 * Enable/Disable interactive element selection for the mesh. Useful to disable element selection on procedural meshes where the selection will be frequently invalidated.
+	 * 
+	 * @param bInIsElementSelectable Whether elements of the dynamic mesh will be interactively selectable.
+	 */
+	UFUNCTION(BlueprintCallable, Category = "Dynamic Mesh Component")
+	GEOMETRYFRAMEWORK_API void SetIsElementSelectable(bool bInIsElementSelectable) { bIsElementSelectable = bInIsElementSelectable; }
+
 protected:
 	/**
 	 * Internal FDynamicMesh is stored inside a UDynamicMesh container, which allows it to be
@@ -415,6 +429,11 @@ protected:
 	 */
 	bool bIsEditable = true;
 
+	/**
+	 * If set to false, the element selection UI will not be used for this mesh component. This is useful for procedural meshes where selection would be frequently invalidated.
+	 */
+	bool bIsElementSelectable = true;
+
 
 
 	//===============================================================================================================
