change IsElementSelectable name to AllowsGeometrySelection

diff --git a/Engine/Plugins/Editor/ModelingToolsEditorMode/Source/ModelingToolsEditorMode/Private/ModelingToolsEditorMode.cpp b/Engine/Plugins/Editor/ModelingToolsEditorMode/Source/ModelingToolsEditorMode/Private/ModelingToolsEditorMode.cpp
--- a/Engine/Plugins/Editor/ModelingToolsEditorMode/Source/ModelingToolsEditorMode/Private/ModelingToolsEditorMode.cpp
+++ b/Engine/Plugins/Editor/ModelingToolsEditorMode/Source/ModelingToolsEditorMode/Private/ModelingToolsEditorMode.cpp
@@ -1517,7 +1517,7 @@ void UModelingToolsEditorMode::UpdateSelectionManagerOnEditorSelectionChange(boo
 	}
 
 	// filter out any dynamic mesh components that aren't editable or aren't element-selectable
-	SelectedDynamicMeshComponents.SetNum(Algo::RemoveIf(SelectedDynamicMeshComponents, [](UDynamicMeshComponent* DMC)->bool {return !DMC->IsEditable() || !DMC->IsElementSelectable();}));
+	SelectedDynamicMeshComponents.SetNum(Algo::RemoveIf(SelectedDynamicMeshComponents, [](UDynamicMeshComponent* DMC)->bool {return !DMC->IsEditable() || !DMC->AllowsGeometrySelection();}));
 
 	// convert selected Component types into selection Identifiers
 	TArray<FGeometryIdentifier> ValidIdentifiers;
diff --git a/Engine/Plugins/Runtime/GeometryScripting/Source/GeometryScriptingEditor/Private/GeometryActors/GeneratedDynamicMeshActor.cpp b/Engine/Plugins/Runtime/GeometryScripting/Source/GeometryScriptingEditor/Private/GeometryActors/GeneratedDynamicMeshActor.cpp
--- a/Engine/Plugins/Runtime/GeometryScripting/Source/GeometryScriptingEditor/Private/GeometryActors/GeneratedDynamicMeshActor.cpp
+++ b/Engine/Plugins/Runtime/GeometryScripting/Source/GeometryScriptingEditor/Private/GeometryActors/GeneratedDynamicMeshActor.cpp
@@ -51,7 +51,7 @@ void AGeneratedDynamicMeshActor::PostActorCreated()
 	// Default newly-created generated dynamic meshes to not be element-selectable
 	if (UDynamicMeshComponent* Component = GetDynamicMeshComponent())
 	{
-		Component->SetIsElementSelectable(false);
+		Component->SetAllowsGeometrySelection(false);
 	}
 }
 
diff --git a/Engine/Source/Runtime/GeometryFramework/Public/Components/DynamicMeshComponent.h b/Engine/Source/Runtime/GeometryFramework/Public/Components/DynamicMeshComponent.h
--- a/Engine/Source/Runtime/GeometryFramework/Public/Components/DynamicMeshComponent.h
+++ b/Engine/Source/Runtime/GeometryFramework/Public/Components/DynamicMeshComponent.h
@@ -235,17 +235,17 @@ public:
 
 
 	/**
-	 * @return Whether elements of the dynamic mesh can be interactively selected.
+	 * @return Whether geometry elements (triangles, edges, etc) of the dynamic mesh can be interactively selected.
 	 */
 	UFUNCTION(BlueprintCallable, Category = "Dynamic Mesh Component")
-	GEOMETRYFRAMEWORK_API bool IsElementSelectable() const { return bIsElementSelectable; }
+	GEOMETRYFRAMEWORK_API bool AllowsGeometrySelection() const { return bAllowsGeometrySelection; }
 	/**
-	 * Enable/Disable interactive element selection for the mesh. Useful to disable element selection on procedural meshes where the selection will be frequently invalidated.
+	 * Enable/Disable interactive geometry element selection for the mesh. Useful to disable geometry selection on procedural meshes where the selection will be frequently invalidated.
 	 * 
-	 * @param bInIsElementSelectable Whether elements of the dynamic mesh will be interactively selectable.
+	 * @param bInAllowsGeometrySelection Whether geometry elements (triangles, edges, etc) of the dynamic mesh will be interactively selectable.
 	 */
 	UFUNCTION(BlueprintCallable, Category = "Dynamic Mesh Component")
-	GEOMETRYFRAMEWORK_API void SetIsElementSelectable(bool bInIsElementSelectable) { bIsElementSelectable = bInIsElementSelectable; }
+	GEOMETRYFRAMEWORK_API void SetAllowsGeometrySelection(bool bInAllowsGeometrySelection) { bAllowsGeometrySelection = bInAllowsGeometrySelection; }
 
 protected:
 	/**
@@ -429,10 +429,12 @@ protected:
 	 */
 	bool bIsEditable = true;
 
+private:
 	/**
-	 * If set to false, the element selection UI will not be used for this mesh component. This is useful for procedural meshes where selection would be frequently invalidated.
+	 * If set to false, the geometry element selection UI will not be used for this mesh component. This is useful for procedural meshes where selection would be frequently invalidated.
 	 */
-	bool bIsElementSelectable = true;
+	UPROPERTY()
+	bool bAllowsGeometrySelection = true;
 
 
 
