implement the uprimitivecomponent material slot name interfaces on base dynamic mesh component so that dynamic mesh components work with code/BPs that access component materials via slot names

diff --git a/Engine/Source/Runtime/GeometryFramework/Private/Components/BaseDynamicMeshComponent.cpp b/Engine/Source/Runtime/GeometryFramework/Private/Components/BaseDynamicMeshComponent.cpp
--- a/Engine/Source/Runtime/GeometryFramework/Private/Components/BaseDynamicMeshComponent.cpp
+++ b/Engine/Source/Runtime/GeometryFramework/Private/Components/BaseDynamicMeshComponent.cpp
@@ -247,6 +247,55 @@ FMaterialRelevance UBaseDynamicMeshComponent::GetMaterialRelevance(ERHIFeatureLe
 	return Result;
 }
 
+// Note Dynamic Meshes don't really have named material slots, so we generate slot names from material + index
+namespace UE::Private::DynamicMeshMaterialSlotNameHelper
+{
+	static FName GetMaterialSlotName(const UMaterialInterface* Mat, int32 MaterialIndex)
+	{
+		return FName(FString::Printf(TEXT("%s_%d"), *((Mat) ? Mat->GetName() : TEXT("Material")), MaterialIndex));
+	}
+}
+TArray<FName> UBaseDynamicMeshComponent::GetMaterialSlotNames() const
+{
+	using namespace UE::Private::DynamicMeshMaterialSlotNameHelper;
+	TArray<FName> ToRet;
+	const int32 NumMaterials = GetNumMaterials();
+	ToRet.Reserve(NumMaterials);
+	for (int32 Idx = 0; Idx < NumMaterials; ++Idx)
+	{
+		ToRet.Add(GetMaterialSlotName(GetMaterial(Idx), Idx));
+	}
+	return ToRet;
+}
+
+bool UBaseDynamicMeshComponent::IsMaterialSlotNameValid(FName MaterialSlotName) const
+{
+	using namespace UE::Private::DynamicMeshMaterialSlotNameHelper;
+	const int32 NumMaterials = GetNumMaterials();
+	for (int32 Idx = 0; Idx < NumMaterials; ++Idx)
+	{
+		if (MaterialSlotName == GetMaterialSlotName(GetMaterial(Idx), Idx))
+		{
+			return true;
+		}
+	}
+	return false;
+}
+
+UMaterialInterface* UBaseDynamicMeshComponent::GetMaterialByName(FName MaterialSlotName) const
+{
+	using namespace UE::Private::DynamicMeshMaterialSlotNameHelper;
+	const int32 NumMaterials = GetNumMaterials();
+	for (int32 Idx = 0; Idx < BaseMaterials.Num(); ++Idx)
+	{
+		if (MaterialSlotName == GetMaterialSlotName(GetMaterial(Idx), Idx))
+		{
+			return GetMaterial(Idx);
+		}
+	}
+	return nullptr;
+}
+
 void UBaseDynamicMeshComponent::SetMaterial(int32 ElementIndex, UMaterialInterface* Material)
 {
 	check(ElementIndex >= 0);
diff --git a/Engine/Source/Runtime/GeometryFramework/Public/Components/BaseDynamicMeshComponent.h b/Engine/Source/Runtime/GeometryFramework/Public/Components/BaseDynamicMeshComponent.h
--- a/Engine/Source/Runtime/GeometryFramework/Public/Components/BaseDynamicMeshComponent.h
+++ b/Engine/Source/Runtime/GeometryFramework/Public/Components/BaseDynamicMeshComponent.h
@@ -594,6 +594,15 @@ public:
 
 	GEOMETRYFRAMEWORK_API virtual void SetNumMaterials(int32 NumMaterials);
 
+	//~ Dynamic Mesh component just has an array of materials without managing slot names, but some methods expect
+	//~ to access materials via slot name, so we still implement the UPrimitiveComponent interface to do so
+	// @return an array of slot names generated from the current materials
+	GEOMETRYFRAMEWORK_API virtual TArray<FName> GetMaterialSlotNames() const override;
+	// @return true if the dynamic mesh has a material with this slot name
+	GEOMETRYFRAMEWORK_API virtual bool IsMaterialSlotNameValid(FName MaterialSlotName) const override;
+	// @return a material on this component with this slot name, or null not found
+	GEOMETRYFRAMEWORK_API virtual UMaterialInterface* GetMaterialByName(FName MaterialSlotName) const override;
+
 	//~ UObject Interface.
 #if WITH_EDITOR
 	GEOMETRYFRAMEWORK_API void PostEditChangeProperty(FPropertyChangedEvent& PropertyChangedEvent) override;
