mitigate copy paste not working for dynamic meshes > 1000 tris in some cases, by raising the limit to 200,000 tris before we change to the 'faster but broken in half the cases' method, and surfacing the failure with a toast

diff --git a/Engine/Source/Runtime/GeometryFramework/GeometryFramework.Build.cs b/Engine/Source/Runtime/GeometryFramework/GeometryFramework.Build.cs
--- a/Engine/Source/Runtime/GeometryFramework/GeometryFramework.Build.cs
+++ b/Engine/Source/Runtime/GeometryFramework/GeometryFramework.Build.cs
@@ -21,6 +21,7 @@ public class GeometryFramework : ModuleRules
 				new string[]
 				{
 					"EditorWidgets",
+					"Slate"
 				}
 			);
 				
diff --git a/Engine/Source/Runtime/GeometryFramework/Private/UDynamicMesh.cpp b/Engine/Source/Runtime/GeometryFramework/Private/UDynamicMesh.cpp
--- a/Engine/Source/Runtime/GeometryFramework/Private/UDynamicMesh.cpp
+++ b/Engine/Source/Runtime/GeometryFramework/Private/UDynamicMesh.cpp
@@ -13,15 +13,24 @@
 #include "Engine/Engine.h"
 #include "Generators/MinimalBoxMeshGenerator.h"
 
+
+#if WITH_EDITOR
+#include "Widgets/Notifications/SNotificationList.h"
+#include "Framework/Notifications/NotificationManager.h"
+#endif
+
 #include UE_INLINE_GENERATED_CPP_BY_NAME(UDynamicMesh)
 
+
+#define LOCTEXT_NAMESPACE "UDynamicMesh"
+
 using namespace UE::Geometry;
 
 
 // these cvars are used to support T3D encoding of the internal FDynamicMesh3, see ::ExportCustomProperties() and ::ImportCustomProperties()
 static TAutoConsoleVariable<int32> CVarDynamicMeshTextBasedDupeTriangleCountThreshold(
 	TEXT("geometry.DynamicMesh.TextBasedDupeTriThreshold"),
-	1000,
+	200000,
 	TEXT("Triangle count threshold for text-based UDynamicMesh duplication using Base64. Large values are quite slow."));
 static TAutoConsoleVariable<int32> CVarDynamicMeshDupeHelperTimeout(
 	TEXT("geometry.DynamicMesh.DupeStashTimeout"),
@@ -30,6 +39,21 @@ static TAutoConsoleVariable<int32> CVarDynamicMeshDupeHelperTimeout(
 
 
 
+namespace UE::Private::UDynamicMeshLocal
+{
+	static void DisplayCriticalWarningMessage(const FText& InMessage, float ExpireDuration = 5.0f)
+	{
+#if WITH_EDITOR
+		FNotificationInfo Info(InMessage);
+		Info.ExpireDuration = ExpireDuration;
+		FSlateNotificationManager::Get().AddNotification(Info);
+#endif
+
+		UE_LOG(LogGeometry, Warning, TEXT("%s"), *InMessage.ToString());
+	}
+}
+
+
 UDynamicMesh::UDynamicMesh(const FObjectInitializer& ObjectInitializer)
 	: Super(ObjectInitializer)
 {
@@ -541,6 +565,7 @@ void UDynamicMesh::ImportCustomProperties(const TCHAR* SourceText, FFeedbackCont
 		}
 
 		// if we got here we failed. Rather than produce an empty mesh, we generate a small cube
+		UE::Private::UDynamicMeshLocal::DisplayCriticalWarningMessage(LOCTEXT("DynamicMeshPasteFailed", "Dynamic Mesh paste failed! See log for details."));
 		UE_LOG(LogGeometry, Warning, TEXT("UDynamicMesh text-based property serialization incomplete, generating box as placeholder. Try increasing geometry.DynamicMesh.TextBasedDupeTriThreshold, or geometry.DynamicMesh.DupeStashTimeout."))
 		FMinimalBoxMeshGenerator BoxGen;
 		BoxGen.Box = UE::Geometry::FOrientedBox3d(FVector3d::Zero(), 50.0 * FVector3d::One());
@@ -622,3 +647,5 @@ void UDynamicMeshPool::FreeAllMeshes()
 	CachedMeshes.Reset();
 	AllCreatedMeshes.Reset();
 }
+
+#undef LOCTEXT_NAMESPACE
\ No newline at end of file
