Fixed an issue with ProceduralMeshComponent and DynamicMeshComponent that left behind ghost collision after clearing mesh sections w/ bUseAsyncCooking.

diff --git a/Engine/Plugins/Runtime/ProceduralMeshComponent/Source/ProceduralMeshComponent/Private/ProceduralMeshComponent.cpp b/Engine/Plugins/Runtime/ProceduralMeshComponent/Source/ProceduralMeshComponent/Private/ProceduralMeshComponent.cpp
--- a/Engine/Plugins/Runtime/ProceduralMeshComponent/Source/ProceduralMeshComponent/Private/ProceduralMeshComponent.cpp
+++ b/Engine/Plugins/Runtime/ProceduralMeshComponent/Source/ProceduralMeshComponent/Private/ProceduralMeshComponent.cpp
@@ -1067,7 +1067,18 @@ void UProceduralMeshComponent::FinishPhysicsAsyncCook(bool bSuccess, UBodySetup*
 	int32 FoundIdx;
 	if(AsyncBodySetupQueue.Find(FinishedBodySetup, FoundIdx))
 	{
-		if (bSuccess)
+		// Note: currently no-cook-needed is reported identically to cook failed.
+		// Checking AggGeom.ConvexElems & ContainsPhysicsTriMeshData here is a hack to distinguish the no-cook-needed case
+		// These checks mirror those from UBodySetup::GetCookInfo.
+		// TODO: remove this hack to distinguish the no-cook-needed case when/if that is no longer identical to the cook failed case
+		const ECollisionTraceFlag BodyCollisionType = FinishedBodySetup->GetCollisionTraceFlag();
+		const bool bEmptySimpleCollision = FinishedBodySetup->AggGeom.ConvexElems.Num() == 0;
+		const bool bEmptyComplexCollision = !ContainsPhysicsTriMeshData(false);
+		const bool bNoCookNeeded =
+			(BodyCollisionType == CTF_UseSimpleAsComplex && bEmptySimpleCollision) ||
+			(BodyCollisionType == CTF_UseComplexAsSimple && bEmptyComplexCollision) ||
+			(BodyCollisionType == CTF_UseSimpleAndComplex && bEmptySimpleCollision && bEmptyComplexCollision);
+		if (bSuccess || bNoCookNeeded)
 		{
 			//The new body was found in the array meaning it's newer so use it
 			ProcMeshBodySetup = FinishedBodySetup;
diff --git a/Engine/Source/Runtime/GeometryFramework/Private/Components/DynamicMeshComponent.cpp b/Engine/Source/Runtime/GeometryFramework/Private/Components/DynamicMeshComponent.cpp
--- a/Engine/Source/Runtime/GeometryFramework/Private/Components/DynamicMeshComponent.cpp
+++ b/Engine/Source/Runtime/GeometryFramework/Private/Components/DynamicMeshComponent.cpp
@@ -1738,9 +1738,17 @@ void UDynamicMeshComponent::FinishPhysicsAsyncCook(bool bSuccess, UBodySetup* Fi
 	if (AsyncBodySetupQueue.Find(FinishedBodySetup, FoundIdx))
 	{
 		// Note: currently no-cook-needed is reported identically to cook failed.
-		// Checking AggGeom.GetElemCount() here is a hack to distinguish the no-cook-needed case
+		// Checking AggGeom.ConvexElems & ContainsPhysicsTriMeshData here is a hack to distinguish the no-cook-needed case
+		// These checks mirror those from UBodySetup::GetCookInfo.
 		// TODO: remove this hack to distinguish the no-cook-needed case when/if that is no longer identical to the cook failed case
-		if (bSuccess || FinishedBodySetup->AggGeom.GetElementCount() > 0)
+		const ECollisionTraceFlag BodyCollisionType = FinishedBodySetup->GetCollisionTraceFlag();
+		const bool bEmptySimpleCollision = FinishedBodySetup->AggGeom.ConvexElems.Num() == 0;
+		const bool bEmptyComplexCollision = !ContainsPhysicsTriMeshData(false);
+		const bool bNoCookNeeded =
+			(BodyCollisionType == CTF_UseSimpleAsComplex && bEmptySimpleCollision) ||
+			(BodyCollisionType == CTF_UseComplexAsSimple && bEmptyComplexCollision) ||
+			(BodyCollisionType == CTF_UseSimpleAndComplex && bEmptySimpleCollision && bEmptyComplexCollision);
+		if (bSuccess || bNoCookNeeded)
 		{
 			// The new body was found in the array meaning it's newer, so use it
 			MeshBodySetup = FinishedBodySetup;
