Allow to override Wireframe Material of Dynamicmesh

diff --git a/Engine/Plugins/Runtime/MeshModelingToolset/Source/ModelingComponents/Private/PreviewMesh.cpp b/Engine/Plugins/Runtime/MeshModelingToolset/Source/ModelingComponents/Private/PreviewMesh.cpp
--- a/Engine/Plugins/Runtime/MeshModelingToolset/Source/ModelingComponents/Private/PreviewMesh.cpp
+++ b/Engine/Plugins/Runtime/MeshModelingToolset/Source/ModelingComponents/Private/PreviewMesh.cpp
@@ -167,7 +167,29 @@ void UPreviewMesh::ClearSecondaryRenderMaterial()
 	return DynamicMeshComponent->ClearSecondaryRenderMaterial();
 }
 
+void UPreviewMesh::SetOverrideWireframeRenderMaterial(UMaterialInterface* Material)
+{
+	check(DynamicMeshComponent);
+	return DynamicMeshComponent->SetOverrideWireframeRenderMaterial(Material);
+}
+
+void UPreviewMesh::ClearOverrideWireframeRenderMaterial()
+{
+	check(DynamicMeshComponent);
+	return DynamicMeshComponent->ClearOverrideWireframeRenderMaterial();
+}
 
+void UPreviewMesh::SetOverrideSecondaryWireframeRenderMaterial(UMaterialInterface* Material)
+{
+	check(DynamicMeshComponent);
+	return DynamicMeshComponent->SetOverrideSecondaryWireframeRenderMaterial(Material);
+}
+
+void UPreviewMesh::ClearOverrideSecondaryWireframeRenderMaterial()
+{
+	check(DynamicMeshComponent);
+	return DynamicMeshComponent->ClearOverrideSecondaryWireframeRenderMaterial();
+}
 
 void UPreviewMesh::EnableSecondaryTriangleBuffers(TUniqueFunction<bool(const FDynamicMesh3*, int32)> TriangleFilterFunc)
 {
diff --git a/Engine/Plugins/Runtime/MeshModelingToolset/Source/ModelingComponents/Public/PreviewMesh.h b/Engine/Plugins/Runtime/MeshModelingToolset/Source/ModelingComponents/Public/PreviewMesh.h
--- a/Engine/Plugins/Runtime/MeshModelingToolset/Source/ModelingComponents/Public/PreviewMesh.h
+++ b/Engine/Plugins/Runtime/MeshModelingToolset/Source/ModelingComponents/Public/PreviewMesh.h
@@ -157,6 +157,25 @@ public:
 	 */
 	void ClearSecondaryRenderMaterial();
 
+	/**
+	 * Set a override wireframe material for the preview mesh. This material will be applied when wireframe is enabled.
+	 */
+	void SetOverrideWireframeRenderMaterial(UMaterialInterface* Material);
+
+	/**
+	 * Clear the wireframe override material for the preview mesh.
+	 */
+	void ClearOverrideWireframeRenderMaterial();
+
+	/**
+	 * Set an override secondary wireframe material for the preview mesh. This material will be applied to seccondary triangle buffer when wireframe is enabled.
+	 */
+	void SetOverrideSecondaryWireframeRenderMaterial(UMaterialInterface* Material);
+
+	/**
+	 * Clear the wireframe override material for the preview mesh.
+	 */
+	void ClearOverrideSecondaryWireframeRenderMaterial();
 
 	/**
 	 * Enable secondary triangle buffers. The Secondary material will be applied to any triangles that pass TriangleFilterFunc.
diff --git a/Engine/Source/Runtime/GeometryFramework/Private/Components/BaseDynamicMeshComponent.cpp b/Engine/Source/Runtime/GeometryFramework/Private/Components/BaseDynamicMeshComponent.cpp
--- a/Engine/Source/Runtime/GeometryFramework/Private/Components/BaseDynamicMeshComponent.cpp
+++ b/Engine/Source/Runtime/GeometryFramework/Private/Components/BaseDynamicMeshComponent.cpp
@@ -97,7 +97,41 @@ void UBaseDynamicMeshComponent::ClearSecondaryRenderMaterial()
 	}
 }
 
+void UBaseDynamicMeshComponent::SetOverrideWireframeRenderMaterial(UMaterialInterface* Material)
+{
+	if (WireframeMaterialOverride != Material)
+	{
+		WireframeMaterialOverride = Material;
+		NotifyMaterialSetUpdated();
+	}
+}
+
+void UBaseDynamicMeshComponent::ClearOverrideWireframeRenderMaterial()
+{
+	if (WireframeMaterialOverride != nullptr)
+	{
+		WireframeMaterialOverride = nullptr;
+		NotifyMaterialSetUpdated();
+	}
+}
 
+void UBaseDynamicMeshComponent::SetOverrideSecondaryWireframeRenderMaterial(UMaterialInterface* Material)
+{
+	if (SecondaryWireframeMaterialOverride != Material)
+	{
+		SecondaryWireframeMaterialOverride = Material;
+		NotifyMaterialSetUpdated();
+	}
+}
+
+void UBaseDynamicMeshComponent::ClearOverrideSecondaryWireframeRenderMaterial()
+{
+	if (SecondaryWireframeMaterialOverride != nullptr)
+	{
+		SecondaryWireframeMaterialOverride = nullptr;
+		NotifyMaterialSetUpdated();
+	}
+}
 
 void UBaseDynamicMeshComponent::SetSecondaryBuffersVisibility(bool bSecondaryVisibility)
 {
diff --git a/Engine/Source/Runtime/GeometryFramework/Private/Components/BaseDynamicMeshSceneProxy.cpp b/Engine/Source/Runtime/GeometryFramework/Private/Components/BaseDynamicMeshSceneProxy.cpp
--- a/Engine/Source/Runtime/GeometryFramework/Private/Components/BaseDynamicMeshSceneProxy.cpp
+++ b/Engine/Source/Runtime/GeometryFramework/Private/Components/BaseDynamicMeshSceneProxy.cpp
@@ -223,19 +223,26 @@ void FBaseDynamicMeshSceneProxy::GetDynamicMeshElements(const TArray<const FScen
 	FMaterialRenderProxy* WireframeMaterialProxy = nullptr;
 	if (bWireframe)
 	{
-		UMaterialInterface* WireframeMaterial = UBaseDynamicMeshComponent::GetDefaultWireframeMaterial_RenderThread();
-		if (WireframeMaterial != nullptr)
+		if(ParentBaseComponent->HasOverrideWireframeRenderMaterial())
 		{
-			FLinearColor UseWireframeColor = (bProxyIsSelected && (bWantWireframeOnShaded == false || bIsWireframeViewMode))  ?
-				GEngine->GetSelectedMaterialColor() : ParentBaseComponent->WireframeColor;
-			FColoredMaterialRenderProxy* WireframeMaterialInstance = new FColoredMaterialRenderProxy(
-				WireframeMaterial->GetRenderProxy(), UseWireframeColor);
-			Collector.RegisterOneFrameMaterialProxy(WireframeMaterialInstance);
-			WireframeMaterialProxy = WireframeMaterialInstance;
+			WireframeMaterialProxy = ParentBaseComponent->GetOverrideWireframeRenderMaterial()->GetRenderProxy();
 		}
 		else
 		{
-			bWireframe = false;
+			UMaterialInterface* WireframeMaterial = UBaseDynamicMeshComponent::GetDefaultWireframeMaterial_RenderThread();
+			if (WireframeMaterial != nullptr)
+			{
+				FLinearColor UseWireframeColor = (bProxyIsSelected && (bWantWireframeOnShaded == false || bIsWireframeViewMode))  ?
+					GEngine->GetSelectedMaterialColor() : ParentBaseComponent->WireframeColor;
+				FColoredMaterialRenderProxy* WireframeMaterialInstance = new FColoredMaterialRenderProxy(
+					WireframeMaterial->GetRenderProxy(), UseWireframeColor);
+				Collector.RegisterOneFrameMaterialProxy(WireframeMaterialInstance);
+				WireframeMaterialProxy = WireframeMaterialInstance;
+			}
+			else
+			{
+				bWireframe = false;
+			}
 		}
 	}
 
@@ -326,7 +333,9 @@ void FBaseDynamicMeshSceneProxy::GetDynamicMeshElements(const TArray<const FScen
 						{
 							DrawBatch(Collector, *BufferSet, BufferSet->SecondaryIndexBuffer, UseSecondaryMaterialProxy, /*bWireframe*/false, DepthPriority, ViewIndex, DynamicPrimitiveUniformBuffer);
 						}
-						DrawBatch(Collector, *BufferSet, BufferSet->SecondaryIndexBuffer, UseSecondaryMaterialProxy, /*bWireframe*/true, DepthPriority, ViewIndex, DynamicPrimitiveUniformBuffer);
+						FMaterialRenderProxy* UseSecondaryMaterialProxyWireFrame = ParentBaseComponent->HasOverrideSecondaryWireframeRenderMaterial() ? 
+							ParentBaseComponent->GetOverrideSecondaryWireframeRenderMaterial()->GetRenderProxy() : UseSecondaryMaterialProxy;
+						DrawBatch(Collector, *BufferSet, BufferSet->SecondaryIndexBuffer, UseSecondaryMaterialProxyWireFrame, /*bWireframe*/true, DepthPriority, ViewIndex, DynamicPrimitiveUniformBuffer);
 					}
 					else
 					{
diff --git a/Engine/Source/Runtime/GeometryFramework/Public/Components/BaseDynamicMeshComponent.h b/Engine/Source/Runtime/GeometryFramework/Public/Components/BaseDynamicMeshComponent.h
--- a/Engine/Source/Runtime/GeometryFramework/Public/Components/BaseDynamicMeshComponent.h
+++ b/Engine/Source/Runtime/GeometryFramework/Public/Components/BaseDynamicMeshComponent.h
@@ -488,6 +488,80 @@ protected:
 	bool bDrawSecondaryBuffers = true;
 
 
+	//===============================================================================================================
+	// Wireframe Material Override Support : Allow to override default wireframe material
+	// 
+public:
+	/**
+	 * Set an active wireframe override render Material 
+	 */
+	UFUNCTION(BlueprintCallable, Category = "Dynamic Mesh Component")
+	GEOMETRYFRAMEWORK_API virtual void SetOverrideWireframeRenderMaterial(UMaterialInterface* Material);
+
+	/**
+	 * Clear any active wireframe override render material
+	 */
+	UFUNCTION(BlueprintCallable, Category = "Dynamic Mesh Component")
+	GEOMETRYFRAMEWORK_API virtual void ClearOverrideWireframeRenderMaterial();
+
+	/**
+	 * @return true if a active wireframe override render material is set
+	 */
+	virtual bool HasOverrideWireframeRenderMaterial() const
+	{
+		return WireframeMaterialOverride != nullptr;
+	}
+
+	/**
+	 * @return active wireframe override render material
+	 */
+	UFUNCTION(BlueprintCallable, Category = "Dynamic Mesh Component")
+	virtual UMaterialInterface* GetOverrideWireframeRenderMaterial() const
+	{
+		return WireframeMaterialOverride;
+	}
+
+protected:
+	UPROPERTY();
+	TObjectPtr<UMaterialInterface> WireframeMaterialOverride = nullptr;
+
+
+	//===============================================================================================================
+	// Wireframe Material Secondary Override Support : Allow to override default secondary wireframe material
+	// 
+public:
+	/**
+	 * Set an active secondary wireframe override render material 
+	 */
+	UFUNCTION(BlueprintCallable, Category = "Dynamic Mesh Component")
+	GEOMETRYFRAMEWORK_API virtual void SetOverrideSecondaryWireframeRenderMaterial(UMaterialInterface* Material);
+
+	/**
+	 * Clear any active secondary wireframe override render material
+	 */
+	UFUNCTION(BlueprintCallable, Category = "Dynamic Mesh Component")
+	GEOMETRYFRAMEWORK_API virtual void ClearOverrideSecondaryWireframeRenderMaterial();
+
+	/**
+	 * @return true if a active secondary wireframe override render material is set
+	 */
+	virtual bool HasOverrideSecondaryWireframeRenderMaterial() const
+	{
+		return SecondaryWireframeMaterialOverride != nullptr;
+	}
+
+	/**
+	 * @return active secondary wireframe override render material
+	 */
+	UFUNCTION(BlueprintCallable, Category = "Dynamic Mesh Component")
+	virtual UMaterialInterface* GetOverrideSecondaryWireframeRenderMaterial() const
+	{
+		return SecondaryWireframeMaterialOverride;
+	}
+
+protected:
+	UPROPERTY();
+	TObjectPtr<UMaterialInterface> SecondaryWireframeMaterialOverride = nullptr;
 
 
 	//===============================================================================================================
