add aliases for UDynamicMesh::Reset so it will also show up if the user searches for clear, empty, or reset

diff --git a/Engine/Source/Runtime/GeometryFramework/Public/UDynamicMesh.h b/Engine/Source/Runtime/GeometryFramework/Public/UDynamicMesh.h
--- a/Engine/Source/Runtime/GeometryFramework/Public/UDynamicMesh.h
+++ b/Engine/Source/Runtime/GeometryFramework/Public/UDynamicMesh.h
@@ -109,7 +109,7 @@ public:
 	 * Clear the internal mesh to an empty mesh.
 	 * This *does not* allocate a new mesh, so any existing mesh pointers/refs are still valid
 	 */
-	UFUNCTION(BlueprintCallable, Category = "Dynamic Mesh")
+	UFUNCTION(BlueprintCallable, Category = "Dynamic Mesh", meta = (Keywords = "Clear Empty Delete"))
 	GEOMETRYFRAMEWORK_API UPARAM(DisplayName = "Target") UDynamicMesh* Reset();
 
 	/**
