Prevent serializing Dynamic Mesh mesh data into diffable object archives.

diff --git a/Engine/Source/Runtime/GeometryFramework/Private/UDynamicMesh.cpp b/Engine/Source/Runtime/GeometryFramework/Private/UDynamicMesh.cpp
--- a/Engine/Source/Runtime/GeometryFramework/Private/UDynamicMesh.cpp
+++ b/Engine/Source/Runtime/GeometryFramework/Private/UDynamicMesh.cpp
@@ -12,6 +12,7 @@
 #include "HAL/IConsoleManager.h"
 #include "Engine/Engine.h"
 #include "Generators/MinimalBoxMeshGenerator.h"
+#include "TransactionCommon.h"
 
 
 #if WITH_EDITOR
@@ -315,8 +316,8 @@ void UDynamicMesh::Serialize(FArchive& Ar)
 
 	Super::Serialize(Ar);
 
-	// do not serialize mesh for transactions
-	if (Ar.IsTransacting())
+	// Do not serialize mesh for transactions/diff objects. Mesh data doesn't need to end up in the DiffObject sync so skipping is fine. 
+	if (Ar.IsTransacting() || UE::Transaction::DiffUtil::IsGeneratingDiffableObject(Ar))
 	{
 		return;
 	}
