DynamicMesh: Fix UV layers being lost in the render proxy when using a custom decomposition in a dynamic mesh component.

diff --git a/Engine/Source/Runtime/GeometryFramework/Private/Components/DynamicMeshSceneProxy.h b/Engine/Source/Runtime/GeometryFramework/Private/Components/DynamicMeshSceneProxy.h
--- a/Engine/Source/Runtime/GeometryFramework/Private/Components/DynamicMeshSceneProxy.h
+++ b/Engine/Source/Runtime/GeometryFramework/Private/Components/DynamicMeshSceneProxy.h
@@ -150,10 +150,17 @@ public:
 		bIsSingleBuffer = false;
 
 		const FDynamicMesh3* Mesh = ParentComponent->GetRenderMesh();
+		const FDynamicMeshAttributeSet* Attributes = Mesh->Attributes();
+
 		// find suitable overlays
-		const FDynamicMeshUVOverlay* UVOverlay = Mesh->Attributes()->PrimaryUV();
-		const FDynamicMeshNormalOverlay* NormalOverlay = Mesh->Attributes()->PrimaryNormals();
-		const FDynamicMeshColorOverlay* ColorOverlay = Mesh->Attributes()->PrimaryColors();
+		TArray<const FDynamicMeshUVOverlay*, TInlineAllocator<8>> UVOverlays;
+		UVOverlays.SetNum(Attributes->NumUVLayers());
+		for (int32 k = 0; k < UVOverlays.Num(); ++k)
+		{
+			UVOverlays[k] = Attributes->GetUVLayer(k);
+		}
+		const FDynamicMeshNormalOverlay* NormalOverlay = Attributes->PrimaryNormals();
+		const FDynamicMeshColorOverlay* ColorOverlay = Attributes->PrimaryColors();
 		TUniqueFunction<void(int, int, int, const FVector3f&, FVector3f&, FVector3f&)> TangentsFunc = MakeTangentsFunc();
 
 		// init renderbuffers for each set
@@ -167,7 +174,7 @@ public:
 				RenderBuffers->Triangles = Triangles;
 				InitializeBuffersFromOverlays(RenderBuffers, Mesh,
 					Triangles.Num(), Triangles,
-					UVOverlay, NormalOverlay, ColorOverlay, TangentsFunc);
+					UVOverlays, NormalOverlay, ColorOverlay, TangentsFunc);
 
 				ENQUEUE_RENDER_COMMAND(FDynamicMeshSceneProxyInitializeFromDecomposition)(
 					[RenderBuffers](FRHICommandListImmediate& RHICmdList)
