harden dynamic mesh component against null mesh object case

diff --git a/Engine/Source/Runtime/GeometryFramework/Private/Components/DynamicMeshComponent.cpp b/Engine/Source/Runtime/GeometryFramework/Private/Components/DynamicMeshComponent.cpp
--- a/Engine/Source/Runtime/GeometryFramework/Private/Components/DynamicMeshComponent.cpp
+++ b/Engine/Source/Runtime/GeometryFramework/Private/Components/DynamicMeshComponent.cpp
@@ -152,6 +152,13 @@ void UDynamicMeshComponent::PostLoad()
 void UDynamicMeshComponent::PostEditImport()
 {
 	Super::PostEditImport();
+	
+	// MeshObject should never be null here, but we re-validate that it isn't (similar to PostLoad method, above)
+	if (ensure(MeshObject != nullptr) == false)
+	{
+		MeshObject = NewObject<UDynamicMesh>(this, TEXT("DynamicMesh"));
+		MeshObjectChangedHandle = MeshObject->OnMeshChanged().AddUObject(this, &UDynamicMeshComponent::OnMeshObjectChanged);
+	}
 
 	if (CurrentDistanceField.IsValid() != (DistanceFieldMode != EDynamicMeshComponentDistanceFieldMode::NoDistanceField))
 	{
@@ -188,7 +195,7 @@ void UDynamicMeshComponent::PostEditChangeProperty(FPropertyChangedEvent& Proper
 
 void UDynamicMeshComponent::SetMesh(UE::Geometry::FDynamicMesh3&& MoveMesh)
 {
-	if (ensureMsgf(IsEditable(), TEXT("Attempted to modify the internal mesh of a UDynamicMeshComponent that is not editable")))
+	if (ensure(MeshObject) && ensureMsgf(IsEditable(), TEXT("Attempted to modify the internal mesh of a UDynamicMeshComponent that is not editable")))
 	{
 		MeshObject->SetMesh(MoveTemp(MoveMesh));
 	}
@@ -198,14 +205,17 @@ void UDynamicMeshComponent::SetMesh(UE::Geometry::FDynamicMesh3&& MoveMesh)
 void UDynamicMeshComponent::ProcessMesh(
 	TFunctionRef<void(const UE::Geometry::FDynamicMesh3&)> ProcessFunc ) const
 {
-	MeshObject->ProcessMesh(ProcessFunc);
+	if (MeshObject)
+	{
+		MeshObject->ProcessMesh(ProcessFunc);
+	}
 }
 
 
 void UDynamicMeshComponent::EditMesh(TFunctionRef<void(UE::Geometry::FDynamicMesh3&)> EditFunc,
 										   EDynamicMeshComponentRenderUpdateMode UpdateMode )
 {
-	if (ensureMsgf(IsEditable(), TEXT("Attempted to modify the internal mesh of a UDynamicMeshComponent that is not editable")))
+	if (MeshObject && ensureMsgf(IsEditable(), TEXT("Attempted to modify the internal mesh of a UDynamicMeshComponent that is not editable")))
 	{
 		MeshObject->EditMesh(EditFunc);
 		if (UpdateMode != EDynamicMeshComponentRenderUpdateMode::NoUpdate)
@@ -218,6 +228,11 @@ void UDynamicMeshComponent::EditMesh(TFunctionRef<void(UE::Geometry::FDynamicMes
 
 void UDynamicMeshComponent::SetRenderMeshPostProcessor(TUniquePtr<IRenderMeshPostProcessor> Processor)
 {
+	if (!ensure(MeshObject))
+	{
+		return;
+	}
+
 	RenderMeshPostProcessor = MoveTemp(Processor);
 	if (RenderMeshPostProcessor)
 	{
@@ -262,7 +277,7 @@ const FDynamicMesh3* UDynamicMeshComponent::GetRenderMesh() const
 
 void UDynamicMeshComponent::ApplyTransform(const FTransform3d& Transform, bool bInvert)
 {
-	if (ensureMsgf(IsEditable(), TEXT("Attempted to modify the internal mesh of a UDynamicMeshComponent that is not editable")))
+	if (ensure(MeshObject) && ensureMsgf(IsEditable(), TEXT("Attempted to modify the internal mesh of a UDynamicMeshComponent that is not editable")))
 	{
 		MeshObject->EditMesh([&](FDynamicMesh3& EditMesh)
 		{
@@ -283,7 +298,7 @@ void UDynamicMeshComponent::ApplyTransform(const FTransform3d& Transform, bool b
 bool UDynamicMeshComponent::ValidateMaterialSlots(bool bCreateIfMissing, bool bDeleteExtraSlots)
 {
 	int32 MaxMeshMaterialID = 0;
-	MeshObject->ProcessMesh([&](const FDynamicMesh3& EditMesh)
+	ProcessMesh([&](const FDynamicMesh3& EditMesh)
 	{
 		if (EditMesh.HasAttributes() && EditMesh.Attributes()->HasMaterialID() && EditMesh.Attributes()->GetMaterialID() != nullptr)
 		{
@@ -345,7 +360,7 @@ void UDynamicMeshComponent::InvalidateAutoCalculatedTangents()
 
 const UE::Geometry::FMeshTangentsf* UDynamicMeshComponent::GetAutoCalculatedTangents() 
 { 
-	if (GetTangentsType() == EDynamicMeshComponentTangentsMode::AutoCalculated && GetDynamicMesh()->GetMeshRef().HasAttributes())
+	if (ensure(MeshObject) && GetTangentsType() == EDynamicMeshComponentTangentsMode::AutoCalculated && GetDynamicMesh()->GetMeshRef().HasAttributes())
 	{
 		UpdateAutoCalculatedTangents();
 		return (bAutoCalculatedTangentsValid) ? &AutoCalculatedTangents : nullptr;
@@ -380,7 +395,7 @@ void UDynamicMeshComponent::UpdateAutoCalculatedTangents()
 
 void UDynamicMeshComponent::UpdateLocalBounds()
 {
-	LocalBounds = GetMesh()->GetBounds(true);
+	LocalBounds = MeshObject ? GetMesh()->GetBounds(true) : FAxisAlignedBox3d::Empty();
 	if (LocalBounds.MaxDim() <= 0)
 	{
 		// If bbox is empty, set a very small bbox to avoid log spam/etc in other engine systems.
@@ -417,7 +432,7 @@ void UDynamicMeshComponent::ResetProxy()
 
 void UDynamicMeshComponent::NotifyMeshUpdated()
 {
-	if (RenderMeshPostProcessor)
+	if (MeshObject && RenderMeshPostProcessor)
 	{
 		RenderMeshPostProcessor->ProcessMesh(*GetMesh(), *RenderMesh);
 	}
@@ -433,6 +448,11 @@ void UDynamicMeshComponent::NotifyMeshModified()
 
 void UDynamicMeshComponent::FastNotifyColorsUpdated()
 {
+	if (!ensure(MeshObject))
+	{
+		return;
+	}
+
 	// should not be using fast paths if we have to run mesh postprocessor
 	if (ensure(!RenderMeshPostProcessor) == false)
 	{
@@ -479,6 +499,11 @@ void UDynamicMeshComponent::FastNotifyColorsUpdated()
 
 void UDynamicMeshComponent::FastNotifyPositionsUpdated(bool bNormals, bool bColors, bool bUVs)
 {
+	if (!ensure(MeshObject))
+	{
+		return;
+	}
+
 	// should not be using fast paths if we have to run mesh postprocessor
 	if (ensure(!RenderMeshPostProcessor) == false)
 	{
@@ -516,6 +541,11 @@ void UDynamicMeshComponent::FastNotifyPositionsUpdated(bool bNormals, bool bColo
 
 void UDynamicMeshComponent::FastNotifyVertexAttributesUpdated(bool bNormals, bool bColors, bool bUVs)
 {
+	if (!ensure(MeshObject))
+	{
+		return;
+	}
+
 	// should not be using fast paths if we have to run mesh postprocessor
 	if (ensure(!RenderMeshPostProcessor) == false)
 	{
@@ -542,6 +572,11 @@ void UDynamicMeshComponent::FastNotifyVertexAttributesUpdated(bool bNormals, boo
 
 void UDynamicMeshComponent::FastNotifyVertexAttributesUpdated(EMeshRenderAttributeFlags UpdatedAttributes)
 {
+	if (!ensure(MeshObject))
+	{
+		return;
+	}
+
 	// should not be using fast paths if we have to run mesh postprocessor
 	if (ensure(!RenderMeshPostProcessor) == false)
 	{
@@ -623,6 +658,11 @@ void UDynamicMeshComponent::NotifyMeshVertexAttributesModified( bool bPositions,
 
 void UDynamicMeshComponent::FastNotifySecondaryTrianglesChanged()
 {
+	if (!ensure(MeshObject))
+	{
+		return;
+	}
+
 	// should not be using fast paths if we have to run mesh postprocessor
 	if (ensure(!RenderMeshPostProcessor) == false)
 	{
@@ -646,6 +686,11 @@ void UDynamicMeshComponent::FastNotifySecondaryTrianglesChanged()
 
 void UDynamicMeshComponent::FastNotifyTriangleVerticesUpdated(const TArray<int32>& Triangles, EMeshRenderAttributeFlags UpdatedAttributes)
 {
+	if (!ensure(MeshObject))
+	{
+		return;
+	}
+
 	// should not be using fast paths if we have to run mesh postprocessor
 	if (ensure(!RenderMeshPostProcessor) == false)
 	{
@@ -728,6 +773,11 @@ void UDynamicMeshComponent::FastNotifyTriangleVerticesUpdated(const TArray<int32
 
 void UDynamicMeshComponent::FastNotifyTriangleVerticesUpdated(const TSet<int32>& Triangles, EMeshRenderAttributeFlags UpdatedAttributes)
 {
+	if (!ensure(MeshObject))
+	{
+		return;
+	}
+
 	// should not be using fast paths if we have to run mesh postprocessor
 	if (ensure(!RenderMeshPostProcessor) == false)
 	{
@@ -853,7 +903,7 @@ TFuture<bool> UDynamicMeshComponent::FastNotifyTriangleVerticesUpdated_TryPrecom
 	TArray<int32>& UpdateSetsOut,
 	FAxisAlignedBox3d& BoundsOut)
 {
-	if ((!!RenderMeshPostProcessor) || (GetCurrentSceneProxy() == nullptr) || (!Decomposition) || !AllowFastUpdate())
+	if (!ensure(MeshObject) || (!!RenderMeshPostProcessor) || (GetCurrentSceneProxy() == nullptr) || (!Decomposition) || !AllowFastUpdate())
 	{
 		// is there a simpler way to do this? cannot seem to just make a TFuture<bool>...
 		return Async(DynamicMeshComponentAsyncExecTarget, []() { return false; });
@@ -959,7 +1009,7 @@ FPrimitiveSceneProxy* UDynamicMeshComponent::CreateSceneProxy()
 	ensure(GetCurrentSceneProxy() == nullptr);
 
 	FDynamicMeshSceneProxy* NewProxy = nullptr;
-	if (GetMesh()->TriangleCount() > 0)
+	if (MeshObject && GetMesh()->TriangleCount() > 0)
 	{
 		NewProxy = new FDynamicMeshSceneProxy(this);
 
@@ -1282,19 +1332,28 @@ void UDynamicMeshComponent::SetInvalidateProxyOnChangeEnabled(bool bEnabled)
 void UDynamicMeshComponent::ApplyChange(const FMeshVertexChange* Change, bool bRevert)
 {
 	// will fire UDynamicMesh::MeshChangedEvent, which will call OnMeshObjectChanged() below to invalidate proxy, fire change events, etc
-	MeshObject->ApplyChange(Change, bRevert);
+	if (ensure(MeshObject))
+	{
+		MeshObject->ApplyChange(Change, bRevert);
+	}
 }
 
 void UDynamicMeshComponent::ApplyChange(const FMeshChange* Change, bool bRevert)
 {
 	// will fire UDynamicMesh::MeshChangedEvent, which will call OnMeshObjectChanged() below to invalidate proxy, fire change events, etc
-	MeshObject->ApplyChange(Change, bRevert);
+	if (ensure(MeshObject))
+	{
+		MeshObject->ApplyChange(Change, bRevert);
+	}
 }
 
 void UDynamicMeshComponent::ApplyChange(const FMeshReplacementChange* Change, bool bRevert)
 {
 	// will fire UDynamicMesh::MeshChangedEvent, which will call OnMeshObjectChanged() below to invalidate proxy, fire change events, etc
-	MeshObject->ApplyChange(Change, bRevert);
+	if (ensure(MeshObject))
+	{
+		MeshObject->ApplyChange(Change, bRevert);
+	}
 }
 
 void UDynamicMeshComponent::OnMeshObjectChanged(UDynamicMesh* ChangedMeshObject, FDynamicMeshChangeInfo ChangeInfo)
diff --git a/Engine/Source/Runtime/GeometryFramework/Public/Components/DynamicMeshComponent.h b/Engine/Source/Runtime/GeometryFramework/Public/Components/DynamicMeshComponent.h
--- a/Engine/Source/Runtime/GeometryFramework/Public/Components/DynamicMeshComponent.h
+++ b/Engine/Source/Runtime/GeometryFramework/Public/Components/DynamicMeshComponent.h
@@ -181,13 +181,13 @@ public:
 	 * @return pointer to internal mesh
 	 * @warning avoid usage of this function, access via GetDynamicMesh() instead
 	 */
-	virtual FDynamicMesh3* GetMesh() override { return MeshObject->GetMeshPtr(); }
+	virtual FDynamicMesh3* GetMesh() override { return MeshObject ? MeshObject->GetMeshPtr() : nullptr; }
 
 	/**
 	 * @return pointer to internal mesh
 	 * @warning avoid usage of this function, access via GetDynamicMesh() instead
 	 */
-	virtual const FDynamicMesh3* GetMesh() const override { return MeshObject->GetMeshPtr(); }
+	virtual const FDynamicMesh3* GetMesh() const override { return MeshObject ? MeshObject->GetMeshPtr() : nullptr; }
 
 	/**
 	 * @return the child UDynamicMesh
