make dynamic mesh component broadcast the property changed event for dynamic mesh changes

diff --git a/Engine/Source/Runtime/GeometryFramework/Private/Components/DynamicMeshComponent.cpp b/Engine/Source/Runtime/GeometryFramework/Private/Components/DynamicMeshComponent.cpp
--- a/Engine/Source/Runtime/GeometryFramework/Private/Components/DynamicMeshComponent.cpp
+++ b/Engine/Source/Runtime/GeometryFramework/Private/Components/DynamicMeshComponent.cpp
@@ -21,6 +21,7 @@
 #include "DynamicMesh/MeshTransforms.h"
 
 #include "UObject/UE5ReleaseStreamObjectVersion.h"
+#include "UObject/UObjectGlobals.h"
 
 // default proxy for this component
 #include "Components/DynamicMeshSceneProxy.h"
@@ -1371,6 +1372,7 @@ void UDynamicMeshComponent::OnMeshObjectChanged(UDynamicMesh* ChangedMeshObject,
 		}
 
 		OnMeshChanged.Broadcast();
+		BroadcastMeshPropertyChangeEvent();
 
 		if (ChangeInfo.Type == EDynamicMeshChangeType::MeshVertexChange)
 		{
@@ -1391,12 +1393,25 @@ void UDynamicMeshComponent::OnMeshObjectChanged(UDynamicMesh* ChangedMeshObject,
 			NotifyMeshUpdated();
 		}
 		OnMeshChanged.Broadcast();
+		BroadcastMeshPropertyChangeEvent();
 	}
 
 	InternalOnMeshUpdated();
 }
 
 
+void UDynamicMeshComponent::BroadcastMeshPropertyChangeEvent()
+{
+#if WITH_EDITOR
+	if (FProperty* MeshProperty = UDynamicMeshComponent::StaticClass()->FindPropertyByName(GET_MEMBER_NAME_CHECKED(UDynamicMeshComponent, MeshObject)))
+	{
+		FPropertyChangedEvent MeshChangedEvent(MeshProperty);
+		FCoreUObjectDelegates::OnObjectPropertyChanged.Broadcast(this, MeshChangedEvent);
+	}
+#endif
+}
+
+
 void UDynamicMeshComponent::SetDynamicMesh(UDynamicMesh* NewMesh)
 {
 	if (ensure(NewMesh) == false)
@@ -1417,6 +1432,7 @@ void UDynamicMeshComponent::SetDynamicMesh(UDynamicMesh* NewMesh)
 
 	NotifyMeshUpdated();
 	OnMeshChanged.Broadcast();
+	BroadcastMeshPropertyChangeEvent();
 
 	InternalOnMeshUpdated();
 }
diff --git a/Engine/Source/Runtime/GeometryFramework/Public/Components/DynamicMeshComponent.h b/Engine/Source/Runtime/GeometryFramework/Public/Components/DynamicMeshComponent.h
--- a/Engine/Source/Runtime/GeometryFramework/Public/Components/DynamicMeshComponent.h
+++ b/Engine/Source/Runtime/GeometryFramework/Public/Components/DynamicMeshComponent.h
@@ -881,4 +881,6 @@ private:
 	// Internal helper to be called when mesh data is updated
 	void InternalOnMeshUpdated();
 
+	void BroadcastMeshPropertyChangeEvent();
+
 };
