DynamicMesh: Make the base dynamic mesh scene proxy call into the component to allow it to override the hit proxy.

diff --git a/Engine/Source/Runtime/GeometryFramework/Private/Components/BaseDynamicMeshSceneProxy.cpp b/Engine/Source/Runtime/GeometryFramework/Private/Components/BaseDynamicMeshSceneProxy.cpp
--- a/Engine/Source/Runtime/GeometryFramework/Private/Components/BaseDynamicMeshSceneProxy.cpp
+++ b/Engine/Source/Runtime/GeometryFramework/Private/Components/BaseDynamicMeshSceneProxy.cpp
@@ -493,6 +493,29 @@ void FBaseDynamicMeshSceneProxy::DrawBatch(FMeshElementCollector& Collector, con
 	Collector.AddMesh(ViewIndex, Mesh);
 }
 
+#if WITH_EDITOR
+HHitProxy* FBaseDynamicMeshSceneProxy::CreateHitProxies(UPrimitiveComponent* Component, TArray<TRefCountPtr<HHitProxy>>& OutHitProxies)
+{
+	return FBaseDynamicMeshSceneProxy::CreateHitProxies(Component->GetPrimitiveComponentInterface(), OutHitProxies);
+}
+
+HHitProxy* FBaseDynamicMeshSceneProxy::CreateHitProxies(IPrimitiveComponent* ComponentInterface, TArray<TRefCountPtr<HHitProxy>>& OutHitProxies)
+{
+	// Similar to a static mesh component, we call back into the component to ask it for a hit proxy. In our
+	//  case, we don't have section indices to give it- we're just using this as an optional customization
+	//  point for the dynamic mesh component to create a custom hit proxy without having to use a different
+	//  scene proxy.
+	HHitProxy* HitProxy = ComponentInterface->CreateMeshHitProxy(0, 0);
+	if (HitProxy)
+	{
+		OutHitProxies.Add(HitProxy);
+		return HitProxy;
+	}
+	
+	// Otherwise fall back to base class implementation
+	return FPrimitiveSceneProxy::CreateHitProxies(ComponentInterface, OutHitProxies);
+}
+#endif //WITH_EDITOR
 
 bool FBaseDynamicMeshSceneProxy::AllowStaticDrawPath(const FSceneView* View) const
 {
diff --git a/Engine/Source/Runtime/GeometryFramework/Public/Components/BaseDynamicMeshSceneProxy.h b/Engine/Source/Runtime/GeometryFramework/Public/Components/BaseDynamicMeshSceneProxy.h
--- a/Engine/Source/Runtime/GeometryFramework/Public/Components/BaseDynamicMeshSceneProxy.h
+++ b/Engine/Source/Runtime/GeometryFramework/Public/Components/BaseDynamicMeshSceneProxy.h
@@ -278,7 +278,11 @@ public:
 		int ViewIndex,
 		FDynamicPrimitiveUniformBuffer& DynamicPrimitiveUniformBuffer) const;
 
-
+	// Like FStaticMeshSceneProxy, only override CreateHitProxies in editor, where component may have CreateMeshHitProxy
+#if WITH_EDITOR
+	GEOMETRYFRAMEWORK_API virtual HHitProxy* CreateHitProxies(UPrimitiveComponent* Component, TArray<TRefCountPtr<HHitProxy> >& OutHitProxies) override;
+	GEOMETRYFRAMEWORK_API virtual HHitProxy* CreateHitProxies(IPrimitiveComponent* ComponentInterface, TArray<TRefCountPtr<HHitProxy> >& OutHitProxies) override;
+#endif // WITH_EDITOR
 
 	//
 	// Raytracing APIs
