UDynamicMeshComponent::ApplyTransform: pass the VertexPosition and NormalsTangents flags to EditMesh. Avoids triggering an ensure in UDynamicMeshComponent::FastNotifyVertexAttributesUpdated

diff --git a/Engine/Source/Runtime/GeometryFramework/Private/Components/DynamicMeshComponent.cpp b/Engine/Source/Runtime/GeometryFramework/Private/Components/DynamicMeshComponent.cpp
--- a/Engine/Source/Runtime/GeometryFramework/Private/Components/DynamicMeshComponent.cpp
+++ b/Engine/Source/Runtime/GeometryFramework/Private/Components/DynamicMeshComponent.cpp
@@ -290,7 +290,7 @@ void UDynamicMeshComponent::ApplyTransform(const FTransform3d& Transform, bool b
 			{
 				MeshTransforms::ApplyTransform(EditMesh, Transform, true);
 			}
-		}, EDynamicMeshChangeType::DeformationEdit);
+		}, EDynamicMeshChangeType::DeformationEdit, EDynamicMeshAttributeChangeFlags::VertexPositions | EDynamicMeshAttributeChangeFlags::NormalsTangents, /*bDeferChangeEvents*/ false);
 	}
 }
 
