Add shader resource flag to index buffers if they can be used for raytracing

diff --git a/Engine/Source/Runtime/GeometryFramework/Public/Components/MeshRenderBufferSet.h b/Engine/Source/Runtime/GeometryFramework/Public/Components/MeshRenderBufferSet.h
--- a/Engine/Source/Runtime/GeometryFramework/Public/Components/MeshRenderBufferSet.h
+++ b/Engine/Source/Runtime/GeometryFramework/Public/Components/MeshRenderBufferSet.h
@@ -97,6 +97,11 @@ public:
 	{
 		StaticMeshVertexBuffer.SetUseFullPrecisionUVs(true);
 		StaticMeshVertexBuffer.SetUseHighPrecisionTangentBasis(true);
+
+#if RHI_RAYTRACING
+		EnumAddFlags(IndexBuffer.UsageFlags, EBufferUsageFlags::ShaderResource);
+		EnumAddFlags(SecondaryIndexBuffer.UsageFlags, EBufferUsageFlags::ShaderResource);
+#endif // RHI_RAYTRACING
 	}
 
 	virtual ~FMeshRenderBufferSet()
