Geometry Framework - thread the creation of the renderbuffers with a parallel for.  The render buffer creation is somewhat slow largely due to the math that takes place when computing the tangent space.  This is unavoidable, so the best we can do is thread this.

diff --git a/Engine/Plugins/Runtime/MeshModelingToolset/Source/ModelingComponents/Private/Components/OctreeDynamicMeshSceneProxy.h b/Engine/Plugins/Runtime/MeshModelingToolset/Source/ModelingComponents/Private/Components/OctreeDynamicMeshSceneProxy.h
--- a/Engine/Plugins/Runtime/MeshModelingToolset/Source/ModelingComponents/Private/Components/OctreeDynamicMeshSceneProxy.h
+++ b/Engine/Plugins/Runtime/MeshModelingToolset/Source/ModelingComponents/Private/Components/OctreeDynamicMeshSceneProxy.h
@@ -96,9 +96,12 @@ public:
 			ColorOverlay = Mesh->Attributes()->PrimaryColors();
 		}
 
+		const bool bTrackTriangles = false;
+		const bool bParallel = true;
 		InitializeBuffersFromOverlays(RenderBuffers, Mesh,
 			Mesh->TriangleCount(), Mesh->TriangleIndicesItr(),
-			UVOverlay, NormalOverlay, ColorOverlay, MakeTangentsFunc());
+			UVOverlay, NormalOverlay, ColorOverlay, MakeTangentsFunc(), 
+			bTrackTriangles, bParallel);
 
 		ENQUEUE_RENDER_COMMAND(FOctreeDynamicMeshSceneProxyInitializeSingle)(
 			[this, RenderBuffers](FRHICommandListImmediate& RHICmdList)
@@ -136,9 +139,13 @@ public:
 			FMeshRenderBufferSet* RenderBuffers = AllocateNewRenderBufferSet();
 			RenderBuffers->Material = GetMaterial(0);
 
+			const bool bTrackTriangles = false;
+			const bool bParallel = true;
+
 			InitializeBuffersFromOverlays(RenderBuffers, Mesh,
 				Tris.Num(), Tris,
-				UVOverlay, NormalOverlay, ColorOverlay, TangentsFunc);
+				UVOverlay, NormalOverlay, ColorOverlay, TangentsFunc, 
+				bTrackTriangles, bParallel);
 
 			ENQUEUE_RENDER_COMMAND(FOctreeDynamicMeshSceneProxyInitializeFromDecomposition)(
 				[this, SetID, RenderBuffers](FRHICommandListImmediate& RHICmdList)
diff --git a/Engine/Source/Runtime/GeometryFramework/Private/Components/DynamicMeshSceneProxy.h b/Engine/Source/Runtime/GeometryFramework/Private/Components/DynamicMeshSceneProxy.h
--- a/Engine/Source/Runtime/GeometryFramework/Private/Components/DynamicMeshSceneProxy.h
+++ b/Engine/Source/Runtime/GeometryFramework/Private/Components/DynamicMeshSceneProxy.h
@@ -216,9 +216,12 @@ public:
 		}
 		TUniqueFunction<void(int, int, int, const FVector3f&, FVector3f&, FVector3f&)> TangentsFunc = MakeTangentsFunc();
 
+		const bool bTrackTriangles = false;
+		const bool bParallel = true;
 		InitializeBuffersFromOverlays(RenderBuffers, Mesh,
 			Mesh->TriangleCount(), Mesh->TriangleIndicesItr(),
-			UVOverlays, NormalOverlay, ColorOverlay, TangentsFunc);
+			UVOverlays, NormalOverlay, ColorOverlay, TangentsFunc,
+			bTrackTriangles, bParallel);
 
 		ENQUEUE_RENDER_COMMAND(FDynamicMeshSceneProxyInitializeSingle)(
 			[RenderBuffers](FRHICommandListImmediate& RHICmdList)
@@ -306,9 +309,12 @@ public:
 					}
 				}
 
+				const bool bTrackTriangles = false;
+				const bool bParallel = true;
 				InitializeBuffersFromOverlays(RenderBuffers, Mesh,
 					Triangles.Num(), Triangles,
-					UVOverlays, NormalOverlay, ColorOverlay, TangentsFunc);
+					UVOverlays, NormalOverlay, ColorOverlay, TangentsFunc,
+					bTrackTriangles, bParallel);
 
 				RenderBuffers->Triangles = Triangles;
 
diff --git a/Engine/Source/Runtime/GeometryFramework/Public/Components/BaseDynamicMeshSceneProxy.h b/Engine/Source/Runtime/GeometryFramework/Public/Components/BaseDynamicMeshSceneProxy.h
--- a/Engine/Source/Runtime/GeometryFramework/Public/Components/BaseDynamicMeshSceneProxy.h
+++ b/Engine/Source/Runtime/GeometryFramework/Public/Components/BaseDynamicMeshSceneProxy.h
@@ -115,9 +115,10 @@ public:
 		const FDynamicMeshNormalOverlay* NormalOverlay,
 		const FDynamicMeshColorOverlay* ColorOverlay,
 		TFunctionRef<void(int, int, int, const FVector3f&, FVector3f&, FVector3f&)> TangentsFunc,
-		bool bTrackTriangles = false)
+		bool bTrackTriangles = false,
+		bool bParallel = false)
 	{
-		MeshRenderBufferSetConverter.InitializeBuffersFromOverlays(RenderBuffers, Mesh, NumTriangles, Enumerable, UVOverlay, NormalOverlay, ColorOverlay, TangentsFunc, bTrackTriangles);
+		MeshRenderBufferSetConverter.InitializeBuffersFromOverlays(RenderBuffers, Mesh, NumTriangles, Enumerable, UVOverlay, NormalOverlay, ColorOverlay, TangentsFunc, bTrackTriangles, bParallel);
 	}
 
 
@@ -135,9 +136,10 @@ public:
 		const FDynamicMeshNormalOverlay* NormalOverlay,
 		const FDynamicMeshColorOverlay* ColorOverlay,
 		TFunctionRef<void(int, int, int, const FVector3f&, FVector3f&, FVector3f&)> TangentsFunc,
-		bool bTrackTriangles = false)
+		bool bTrackTriangles = false,
+		bool bParallel = false)
 	{
-		MeshRenderBufferSetConverter.InitializeBuffersFromOverlays(RenderBuffers, Mesh, NumTriangles, Enumerable, UVOverlays, NormalOverlay, ColorOverlay, TangentsFunc, bTrackTriangles);
+		MeshRenderBufferSetConverter.InitializeBuffersFromOverlays(RenderBuffers, Mesh, NumTriangles, Enumerable, UVOverlays, NormalOverlay, ColorOverlay, TangentsFunc, bTrackTriangles, bParallel);
 	}
 
 
diff --git a/Engine/Source/Runtime/GeometryFramework/Public/Components/MeshRenderBufferSet.h b/Engine/Source/Runtime/GeometryFramework/Public/Components/MeshRenderBufferSet.h
--- a/Engine/Source/Runtime/GeometryFramework/Public/Components/MeshRenderBufferSet.h
+++ b/Engine/Source/Runtime/GeometryFramework/Public/Components/MeshRenderBufferSet.h
@@ -4,6 +4,8 @@
 
 #include "CoreMinimal.h"
 
+#include "Async/ParallelFor.h"
+
 #include "Components/BaseDynamicMeshComponent.h"
 #include "Components/DynamicMeshComponent.h"
 #include "DynamicMeshBuilder.h"
@@ -360,12 +362,13 @@ struct FMeshRenderBufferSetConversionUtil
 										const FDynamicMeshNormalOverlay* NormalOverlay,
 										const FDynamicMeshColorOverlay* ColorOverlay,
 										TFunctionRef<void(int, int, int, const FVector3f&, FVector3f&, FVector3f&)> TangentsFunc,
-										bool bTrackTriangles = false)
+										bool bTrackTriangles = false,
+										bool bParallel = false)
 	{
 		TArray<const FDynamicMeshUVOverlay*> UVOverlays;
 		UVOverlays.Add(UVOverlay);
 		InitializeBuffersFromOverlays(RenderBuffers, Mesh, NumTriangles, Enumerable,
-			UVOverlays, NormalOverlay, ColorOverlay, TangentsFunc, bTrackTriangles);
+			UVOverlays, NormalOverlay, ColorOverlay, TangentsFunc, bTrackTriangles, bParallel);
 	}
 
 
@@ -382,7 +385,8 @@ struct FMeshRenderBufferSetConversionUtil
 										const FDynamicMeshNormalOverlay* NormalOverlay,
 										const FDynamicMeshColorOverlay* ColorOverlay,
 										TFunctionRef<void(int, int, int, const FVector3f&, FVector3f&, FVector3f&)> TangentsFunc,
-										bool bTrackTriangles = false);
+										bool bTrackTriangles = false,
+										bool bParallel = false);
 												
 
 
@@ -505,7 +509,8 @@ void FMeshRenderBufferSetConversionUtil::InitializeBuffersFromOverlays( FMeshRen
 																		const FDynamicMeshNormalOverlay* NormalOverlay,
 																		const FDynamicMeshColorOverlay* ColorOverlay,
 																		TFunctionRef<void(int, int, int, const FVector3f&, FVector3f&, FVector3f&)> TangentsFunc,
-																		bool bTrackTriangles)
+																		bool bTrackTriangles,
+																		bool bParallel)
 
 {
 	RenderBuffers->TriangleCount = NumTriangles;
@@ -516,11 +521,9 @@ void FMeshRenderBufferSetConversionUtil::InitializeBuffersFromOverlays( FMeshRen
 
 	bool bHaveColors = (ColorOverlay != nullptr) && (bIgnoreVertexColors == false);
 
-	int NumVertices = NumTriangles * 3;
-	int NumUVOverlays = UVOverlays.Num();
-	int NumTexCoords = FMath::Max(1, NumUVOverlays);		// must have at least one tex coord
-	TArray<FIndex3i, TFixedAllocator<MAX_STATIC_TEXCOORDS>> UVTriangles;
-	UVTriangles.SetNum(NumTexCoords);
+	const int NumVertices = NumTriangles * 3;
+	const int NumUVOverlays = UVOverlays.Num();
+	const int NumTexCoords = FMath::Max(1, NumUVOverlays);		// must have at least one tex coord
 
 	{
 		RenderBuffers->PositionVertexBuffer.Init(NumVertices);
@@ -530,21 +533,34 @@ void FMeshRenderBufferSetConversionUtil::InitializeBuffersFromOverlays( FMeshRen
 	}
 
 	// build triangle list if requested, or if we are using secondary buffers in which case we need it to filter later
-	bool bBuildTriangleList = bTrackTriangles || bUseSecondaryTriBuffers;
-	if (bBuildTriangleList)
+	const bool bBuildTriangleList = bTrackTriangles || bUseSecondaryTriBuffers;
+
+	// populate the triangle array.  we use this for parallelism. 
+	TArray<int32> TriangleArray;
 	{
-		RenderBuffers->Triangles = TArray<int32>();
+		TriangleArray.Reserve(NumTriangles);
+		for (int TriangleID : Enumerable)
+		{
+			TriangleArray.Add(TriangleID);
+		}
 	}
 
-	int TriIdx = 0, VertIdx = 0;
-	FVector3f TangentX, TangentY;
-	for (int TriangleID : Enumerable)
+	TArray<uint32>& IndexBuffer = RenderBuffers->IndexBuffer.Indices;
 	{
-		FIndex3i Tri = Mesh->GetTriangle(TriangleID);
-		for (int32 k = 0; k < NumTexCoords; ++k)
+		// populate the index buffer.
+		for (int idx = 0, IdxMax = 3 * NumTriangles; idx < IdxMax; ++idx)
 		{
-			UVTriangles[k] = (k < NumUVOverlays && UVOverlays[k] != nullptr) ? UVOverlays[k]->GetTriangle(TriangleID) : FIndex3i::Invalid();
+			IndexBuffer[idx] = idx;
 		}
+	}
+
+
+	ParallelFor(TriangleArray.Num(), [&](int idx)
+	{
+		int32 TriangleID = TriangleArray[idx];
+		FIndex3i Tri = Mesh->GetTriangle(TriangleID);
+
+	
 		FIndex3i TriNormal = (NormalOverlay != nullptr) ? NormalOverlay->GetTriangle(TriangleID) : FIndex3i::Zero();
 		FIndex3i TriColor = (ColorOverlay != nullptr) ? ColorOverlay->GetTriangle(TriangleID) : FIndex3i::Zero();
 
@@ -555,6 +571,7 @@ void FMeshRenderBufferSetConversionUtil::InitializeBuffersFromOverlays( FMeshRen
 			bHaveColors = false;
 		}
 
+		int32 VertIdx = idx * 3;
 		for (int j = 0; j < 3; ++j)
 		{
 			RenderBuffers->PositionVertexBuffer.VertexPosition(VertIdx) = (FVector3f)Mesh->GetVertex(Tri[j]);
@@ -571,30 +588,36 @@ void FMeshRenderBufferSetConversionUtil::InitializeBuffersFromOverlays( FMeshRen
 			}
 
 			// get tangents
+			FVector3f TangentX, TangentY;
 			TangentsFunc(Tri[j], TriangleID, j, Normal, TangentX, TangentY);
 
 			RenderBuffers->StaticMeshVertexBuffer.SetVertexTangents(VertIdx, TangentX, TangentY, Normal);
 
-			for (int32 k = 0; k < NumTexCoords; ++k)
-			{
-				FVector2f UV = (UVTriangles[k][j] != FDynamicMesh3::InvalidID) ?
-					UVOverlays[k]->GetElement(UVTriangles[k][j]) : FVector2f::Zero();
-				RenderBuffers->StaticMeshVertexBuffer.SetVertexUV(VertIdx, k, UV);
-			}
-
 			FColor VertexFColor = (bHaveColors && TriColor[j] != FDynamicMesh3::InvalidID) ?
 				GetOverlayColorAsFColor(ColorOverlay, TriColor[j]) : UniformTriColor;
 
 			RenderBuffers->ColorVertexBuffer.VertexColor(VertIdx) = VertexFColor;
 
-			RenderBuffers->IndexBuffer.Indices[TriIdx++] = VertIdx;		// currently TriIdx == VertIdx so we don't really need both...
 			VertIdx++;
 		}
 
-		if (bBuildTriangleList)
+		for (int32 k = 0; k < NumTexCoords; ++k)
 		{
-			RenderBuffers->Triangles->Add(TriangleID);
+			VertIdx = idx * 3;
+			FIndex3i UVTriangle = (k < NumUVOverlays && UVOverlays[k] != nullptr) ? UVOverlays[k]->GetTriangle(TriangleID) : FIndex3i::Invalid();
+			for (int j = 0; j < 3; ++j)
+			{
+				FVector2f UV = (UVTriangle[j] != FDynamicMesh3::InvalidID) ?
+					UVOverlays[k]->GetElement(UVTriangle[j]) : FVector2f::Zero();
+				RenderBuffers->StaticMeshVertexBuffer.SetVertexUV(VertIdx, k, UV);
+				VertIdx++;
+			}
 		}
+	}, !bParallel);
+
+	if (bBuildTriangleList)
+	{
+		RenderBuffers->Triangles = MoveTemp(TriangleArray);
 	}
 
 	// split triangles into secondary buffer (at bit redundant since we just built IndexBuffer, but we may optionally duplicate triangles in the future)
