fix dynamic mesh component using crash on non-compact meshes due to using triangle count instead of max triangle id

diff --git a/Engine/Source/Runtime/GeometryFramework/Private/Components/DynamicMeshComponent.cpp b/Engine/Source/Runtime/GeometryFramework/Private/Components/DynamicMeshComponent.cpp
--- a/Engine/Source/Runtime/GeometryFramework/Private/Components/DynamicMeshComponent.cpp
+++ b/Engine/Source/Runtime/GeometryFramework/Private/Components/DynamicMeshComponent.cpp
@@ -1526,7 +1526,7 @@ bool UDynamicMeshComponent::GetPhysicsTriMeshData(struct FTriMeshCollisionData*
 		else
 		{
 			// map vertices per wedge
-			VertexMap.SetNumZeroed(Mesh.TriangleCount() * 3);
+			VertexMap.SetNumZeroed(Mesh.MaxTriangleID() * 3);
 			// temp array to store the UVs on a vertex (per triangle)
 			TArray<FVector2D> VertUVs;
 			const FDynamicMeshAttributeSet* Attribs = Mesh.Attributes();
