add option for dynamic mesh component to selectively disable physics mesh hit UVs when the project setting is enabled

diff --git a/Engine/Source/Runtime/GeometryFramework/Private/Components/DynamicMeshComponent.cpp b/Engine/Source/Runtime/GeometryFramework/Private/Components/DynamicMeshComponent.cpp
--- a/Engine/Source/Runtime/GeometryFramework/Private/Components/DynamicMeshComponent.cpp
+++ b/Engine/Source/Runtime/GeometryFramework/Private/Components/DynamicMeshComponent.cpp
@@ -1470,7 +1470,7 @@ bool UDynamicMeshComponent::GetTriMeshSizeEstimates(struct FTriMeshCollisionData
 {
 	ProcessMesh([&](const FDynamicMesh3& Mesh)
 		{
-			bool bCopyUVs = UPhysicsSettings::Get()->bSupportUVFromHitResults && Mesh.HasAttributes() && Mesh.Attributes()->NumUVLayers() > 0;
+			bool bCopyUVs = UPhysicsSettings::Get()->bSupportUVFromHitResults && Mesh.HasAttributes() && Mesh.Attributes()->NumUVLayers() > 0 && !bDisableMeshUVHitResults;
 			if (bCopyUVs)
 			{
 				// conservative estimate
@@ -1493,7 +1493,7 @@ bool UDynamicMeshComponent::GetPhysicsTriMeshData(struct FTriMeshCollisionData*
 	ProcessMesh([&](const FDynamicMesh3& Mesh)
 	{
 		// See if we should copy UVs
-		const bool bCopyUVs = UPhysicsSettings::Get()->bSupportUVFromHitResults && Mesh.HasAttributes() && Mesh.Attributes()->NumUVLayers() > 0;
+		const bool bCopyUVs = UPhysicsSettings::Get()->bSupportUVFromHitResults && Mesh.HasAttributes() && Mesh.Attributes()->NumUVLayers() > 0 && !bDisableMeshUVHitResults;
 		if (bCopyUVs)
 		{
 			CollisionData->UVs.SetNum(Mesh.Attributes()->NumUVLayers());
diff --git a/Engine/Source/Runtime/GeometryFramework/Public/Components/DynamicMeshComponent.h b/Engine/Source/Runtime/GeometryFramework/Public/Components/DynamicMeshComponent.h
--- a/Engine/Source/Runtime/GeometryFramework/Public/Components/DynamicMeshComponent.h
+++ b/Engine/Source/Runtime/GeometryFramework/Public/Components/DynamicMeshComponent.h
@@ -777,6 +777,10 @@ public:
 	UPROPERTY(EditAnywhere, BlueprintReadWrite, Category = "Dynamic Mesh Component|Collision");
 	bool bDeferCollisionUpdates = false;
 
+	/** If true, UV hit results will not be enabled for this mesh's complex collision even if they are enabled for the project settings generally. This can make the complex build faster and take less memory. */
+	UPROPERTY(EditAnywhere, BlueprintReadWrite, Category = "Dynamic Mesh Component|Collision");
+	bool bDisableMeshUVHitResults = false;
+
 	/**
 	 * The bDeferCollisionUpdates property is a serialized Component setting that controls when collision is regenerated.
 	 * However, in some cases where a UDynamicMesh is being programmatically edited (eg like a live deformation/etc), we
