make enumerable template parameter pass as const ref to avoid chance of accidentally doing an array copy when called

diff --git a/Engine/Source/Runtime/GeometryFramework/Public/Components/BaseDynamicMeshSceneProxy.h b/Engine/Source/Runtime/GeometryFramework/Public/Components/BaseDynamicMeshSceneProxy.h
--- a/Engine/Source/Runtime/GeometryFramework/Public/Components/BaseDynamicMeshSceneProxy.h
+++ b/Engine/Source/Runtime/GeometryFramework/Public/Components/BaseDynamicMeshSceneProxy.h
@@ -110,7 +110,7 @@ public:
 	void InitializeBuffersFromOverlays(
 		FMeshRenderBufferSet* RenderBuffers,
 		const FDynamicMesh3* Mesh,
-		int NumTriangles, TriangleEnumerable Enumerable,
+		int NumTriangles, const TriangleEnumerable& Enumerable,
 		const FDynamicMeshUVOverlay* UVOverlay,
 		const FDynamicMeshNormalOverlay* NormalOverlay,
 		const FDynamicMeshColorOverlay* ColorOverlay,
@@ -131,7 +131,7 @@ public:
 	void InitializeBuffersFromOverlays(
 		FMeshRenderBufferSet* RenderBuffers,
 		const FDynamicMesh3* Mesh,
-		int NumTriangles, TriangleEnumerable Enumerable,
+		int NumTriangles, const TriangleEnumerable& Enumerable,
 		const TArray<const FDynamicMeshUVOverlay*, UVOverlayListAllocator>& UVOverlays,
 		const FDynamicMeshNormalOverlay* NormalOverlay,
 		const FDynamicMeshColorOverlay* ColorOverlay,
@@ -181,7 +181,7 @@ public:
 	void UpdateVertexBuffersFromOverlays(
 		FMeshRenderBufferSet* RenderBuffers,
 		const FDynamicMesh3* Mesh,
-		int NumTriangles, TriangleEnumerable Enumerable,
+		int NumTriangles, const TriangleEnumerable& Enumerable,
 		const FDynamicMeshNormalOverlay* NormalOverlay,
 		const FDynamicMeshColorOverlay* ColorOverlay,
 		TFunctionRef<void(int, int, int, const FVector3f&, FVector3f&, FVector3f&)> TangentsFunc,
@@ -200,7 +200,7 @@ public:
 	void UpdateVertexUVBufferFromOverlays(
 		FMeshRenderBufferSet* RenderBuffers,
 		const FDynamicMesh3* Mesh,
-		int32 NumTriangles, TriangleEnumerable Enumerable,
+		int32 NumTriangles, const TriangleEnumerable& Enumerable,
 		const TArray<const FDynamicMeshUVOverlay*, UVOverlayListAllocator>& UVOverlays)
 	{
 		MeshRenderBufferSetConverter.UpdateVertexUVBufferFromOverlays(RenderBuffers, Mesh, NumTriangles, Enumerable, UVOverlays);
diff --git a/Engine/Source/Runtime/GeometryFramework/Public/Components/MeshRenderBufferSet.h b/Engine/Source/Runtime/GeometryFramework/Public/Components/MeshRenderBufferSet.h
--- a/Engine/Source/Runtime/GeometryFramework/Public/Components/MeshRenderBufferSet.h
+++ b/Engine/Source/Runtime/GeometryFramework/Public/Components/MeshRenderBufferSet.h
@@ -357,7 +357,7 @@ struct FMeshRenderBufferSetConversionUtil
 	void InitializeBuffersFromOverlays(
 										FMeshRenderBufferSet* RenderBuffers,
 										const FDynamicMesh3* Mesh,
-										int NumTriangles, TriangleEnumerable Enumerable,
+										int NumTriangles, const TriangleEnumerable& Enumerable,
 										const FDynamicMeshUVOverlay* UVOverlay,
 										const FDynamicMeshNormalOverlay* NormalOverlay,
 										const FDynamicMeshColorOverlay* ColorOverlay,
@@ -380,7 +380,7 @@ struct FMeshRenderBufferSetConversionUtil
 	void InitializeBuffersFromOverlays( FMeshRenderBufferSet* RenderBuffers,
 										const FDynamicMesh3* Mesh,
 										int NumTriangles, 
-										TriangleEnumerable Enumerable,
+										const TriangleEnumerable& Enumerable,
 										const TArray<const FDynamicMeshUVOverlay*, UVOverlayListAllocator>& UVOverlays,
 										const FDynamicMeshNormalOverlay* NormalOverlay,
 										const FDynamicMeshColorOverlay* ColorOverlay,
@@ -416,7 +416,7 @@ struct FMeshRenderBufferSetConversionUtil
 	template<typename TriangleEnumerable>
 	void UpdateVertexBuffersFromOverlays( FMeshRenderBufferSet* RenderBuffers,
 										  const FDynamicMesh3* Mesh,
-										  int NumTriangles, TriangleEnumerable Enumerable,
+										  int NumTriangles, const TriangleEnumerable& Enumerable,
 										  const FDynamicMeshNormalOverlay* NormalOverlay,
 										  const FDynamicMeshColorOverlay* ColorOverlay,
 										  TFunctionRef<void(int, int, int, const FVector3f&, FVector3f&, FVector3f&)> TangentsFunc,
@@ -432,7 +432,7 @@ struct FMeshRenderBufferSetConversionUtil
 	template<typename TriangleEnumerable, typename UVOverlayListAllocator>
 	void UpdateVertexUVBufferFromOverlays( FMeshRenderBufferSet* RenderBuffers,
 										   const FDynamicMesh3* Mesh,
-										   int32 NumTriangles, TriangleEnumerable Enumerable,
+										   int32 NumTriangles, const TriangleEnumerable& Enumerable,
 										   const TArray<const FDynamicMeshUVOverlay*, UVOverlayListAllocator>& UVOverlays);
 
 
@@ -504,7 +504,7 @@ template<typename TriangleEnumerable, typename UVOverlayListAllocator>
 void FMeshRenderBufferSetConversionUtil::InitializeBuffersFromOverlays( FMeshRenderBufferSet* RenderBuffers,
 																	    const FDynamicMesh3* Mesh,
 																		int NumTriangles,
-																		TriangleEnumerable Enumerable,
+																		const TriangleEnumerable& Enumerable,
 																	    const TArray<const FDynamicMeshUVOverlay*, UVOverlayListAllocator>& UVOverlays,
 																		const FDynamicMeshNormalOverlay* NormalOverlay,
 																		const FDynamicMeshColorOverlay* ColorOverlay,
@@ -636,7 +636,7 @@ void FMeshRenderBufferSetConversionUtil::InitializeBuffersFromOverlays( FMeshRen
 template<typename TriangleEnumerable>
 void FMeshRenderBufferSetConversionUtil::UpdateVertexBuffersFromOverlays(  FMeshRenderBufferSet* RenderBuffers,
 																		   const FDynamicMesh3* Mesh,
-																		   int NumTriangles, TriangleEnumerable Enumerable,
+																		   int NumTriangles, const TriangleEnumerable& Enumerable,
 																		   const FDynamicMeshNormalOverlay* NormalOverlay,
 																		   const FDynamicMeshColorOverlay* ColorOverlay,
 																		   TFunctionRef<void(int, int, int, const FVector3f&, FVector3f&, FVector3f&)> TangentsFunc,
@@ -716,7 +716,7 @@ void FMeshRenderBufferSetConversionUtil::UpdateVertexBuffersFromOverlays(  FMesh
 template<typename TriangleEnumerable, typename UVOverlayListAllocator>
 void FMeshRenderBufferSetConversionUtil::UpdateVertexUVBufferFromOverlays( FMeshRenderBufferSet* RenderBuffers,
 																	       const FDynamicMesh3* Mesh,
-																	       int32 NumTriangles, TriangleEnumerable Enumerable,
+																	       int32 NumTriangles, const TriangleEnumerable& Enumerable,
 																	       const TArray<const FDynamicMeshUVOverlay*, UVOverlayListAllocator>& UVOverlays)
 {
 	// We align the update to the way we set UV's in InitializeBuffersFromOverlays.
