In the startup of the geometry framework module, skip registering the level editor outliner filter in -server mode like we do for -game

diff --git a/Engine/Source/Runtime/GeometryFramework/Private/GeometryFrameworkModule.cpp b/Engine/Source/Runtime/GeometryFramework/Private/GeometryFrameworkModule.cpp
--- a/Engine/Source/Runtime/GeometryFramework/Private/GeometryFrameworkModule.cpp
+++ b/Engine/Source/Runtime/GeometryFramework/Private/GeometryFrameworkModule.cpp
@@ -19,7 +19,7 @@ void FGeometryFrameworkModule::StartupModule()
 	FCoreDelegates::OnPostEngineInit.AddRaw(this, &FGeometryFrameworkModule::OnPostEngineInit);
 
 #if WITH_EDITOR
-	if (!IsRunningGame())
+	if (!IsRunningGame() && !IsRunningDedicatedServer())
 	{
 		FLevelEditorModule& LevelEditorModule = FModuleManager::LoadModuleChecked<FLevelEditorModule>("LevelEditor");
 
