GeometryFramework: allow asynchronous phyiscs build of UDynamicMeshComponent in editor.

diff --git a/Engine/Source/Runtime/GeometryFramework/Private/Components/DynamicMeshComponent.cpp b/Engine/Source/Runtime/GeometryFramework/Private/Components/DynamicMeshComponent.cpp
--- a/Engine/Source/Runtime/GeometryFramework/Private/Components/DynamicMeshComponent.cpp
+++ b/Engine/Source/Runtime/GeometryFramework/Private/Components/DynamicMeshComponent.cpp
@@ -1621,7 +1621,14 @@ void UDynamicMeshComponent::InvalidatePhysicsData()
 void UDynamicMeshComponent::RebuildPhysicsData()
 {
 	UWorld* World = GetWorld();
-	const bool bUseAsyncCook = World && World->IsGameWorld() && bUseAsyncCooking;
+	const bool bUseAsyncCook = bUseAsyncCooking
+		// It's uncertain whether these checks are actually needed. UProceduralMeshComponent checked for a game,
+		//  not editor, world. It's possible that at the time the editor was not ticked in a way that caused
+		//  physics builds to complete.
+		// We need to allow asynchronous builds in editor since dynamic meshes get used for real time modifications
+		//  that can otherwise hitch as things are dragged. We'll keep the guard against null or Inactive/None worlds
+		//  for now just in case, but it may not be necessary.
+		&& World && (World->IsGameWorld() || World->IsEditorWorld());
 
 	UBodySetup* BodySetup = nullptr;
 	if (bUseAsyncCook)
