make dynamic mesh scene proxy use the standard bounds->cards method, rather than its own version of the same

diff --git a/Engine/Source/Runtime/GeometryFramework/Private/Components/BaseDynamicMeshSceneProxy.cpp b/Engine/Source/Runtime/GeometryFramework/Private/Components/BaseDynamicMeshSceneProxy.cpp
--- a/Engine/Source/Runtime/GeometryFramework/Private/Components/BaseDynamicMeshSceneProxy.cpp
+++ b/Engine/Source/Runtime/GeometryFramework/Private/Components/BaseDynamicMeshSceneProxy.cpp
@@ -768,28 +768,12 @@ void FBaseDynamicMeshSceneProxy::DrawRayTracingBatch(FRayTracingInstanceCollecto
 
 const FCardRepresentationData* FBaseDynamicMeshSceneProxy::GetMeshCardRepresentation() const
 {
-	if (MeshCards.IsValid() && bMeshCardsValid)
-	{
-		return MeshCards.Get();
-	}
-	return nullptr;
+	return MeshCards.Get();
 }
 
-namespace UE::Local
-{
-	// Same as LumenMeshCards::GetAxisAlignedDirection
-	FVector3f GetAxisAlignedDirection(int32 AxisAlignedDirectionIndex, int32& AxisIndex)
-	{
-		AxisIndex = AxisAlignedDirectionIndex / 2;
-		FVector3f Direction(0.0f, 0.0f, 0.0f);
-		Direction[AxisIndex] = AxisAlignedDirectionIndex & 1 ? 1.0f : -1.0f;
-		return Direction;
-	}
-}
 
 void FBaseDynamicMeshSceneProxy::UpdateLumenCardsFromBounds()
 {
-	bMeshCardsValid = false;
 	if (!bVisibleInLumenScene || !UE::DynamicMesh::AllowLumenCardGeneration())
 	{
 		MeshCards.Reset();
@@ -803,69 +787,14 @@ void FBaseDynamicMeshSceneProxy::UpdateLumenCardsFromBounds()
 		MeshCards = MakePimpl<FCardRepresentationData>();
 	}
 
-	*MeshCards = FCardRepresentationData();		 // increments ID
+	*MeshCards = FCardRepresentationData();
 	FMeshCardsBuildData& CardData = MeshCards->MeshCardsBuildData;
 
 	CardData.Bounds = Box;
+	// Mark as two-sided so a high sampling bias is used and hits are accepted even if they don't match well
+	CardData.bMostlyTwoSided = true;
 
-
-	struct FCardDirection
-	{
-		int DirectionIndex;
-		FVector AxisZ;
-		FVector AxisX;
-		FVector AxisY;
-		int AxisZIndex;
-	};
-	TArray<FCardDirection> CardDirections;
-	for (int32 DirectionIndex = 0; DirectionIndex < 6; ++DirectionIndex)
-	{
-		FCardDirection Direction;
-		Direction.DirectionIndex = DirectionIndex;
-		Direction.AxisZ = (FVector)UE::Local::GetAxisAlignedDirection(DirectionIndex, Direction.AxisZIndex);
-		Direction.AxisZ.FindBestAxisVectors(Direction.AxisX, Direction.AxisY);
-		Direction.AxisX = FVector::CrossProduct(Direction.AxisZ, Direction.AxisY);
-		Direction.AxisX.Normalize();
-		CardDirections.Add(Direction);
-	}
-
-	FVector3d Center = Box.GetCenter();
-	FVector3d Extents = Box.GetExtent();
-	float CardOffset = 5.0;
-
-	CardData.CardBuildData.SetNum(CardDirections.Num());
-	for (int32 CardIndex = 0; CardIndex < CardDirections.Num(); ++CardIndex)
-	{
-		FCardDirection Direction = CardDirections[CardIndex];
-		FLumenCardOBBf OBB;
-		OBB.AxisZ = (FVector3f)Direction.AxisZ;
-		OBB.AxisX = (FVector3f)Direction.AxisX;
-		OBB.AxisY = (FVector3f)Direction.AxisY;
-	
-		// project 3D mesh extents onto the specific axes of this CardOBB  (this just reshuffles them but the combinatorics are messy)
-		double ExtentX = FMathd::Abs(Direction.AxisX.Dot( Extents ));
-		double ExtentY = FMathd::Abs(Direction.AxisY.Dot( Extents ));
-		double ExtentZ = FMathd::Abs(Direction.AxisZ.Dot( Extents ));
-
-		// Translate the box along the AxisZ axis so the box center is at the middle of the axis-face 
-		FVector3d LocalCenter = Center;
-		LocalCenter +=  ExtentZ * Direction.AxisZ;
-
-		// hardcoding the card box to cover half the mesh bounds along Z (and full mesh box along X and Y)
-		double CardExtentZ = ExtentZ * 0.5;	
-
-		// shift the card box center so that the +Z face lies on the mesh box face, then bump it forward a bit
-		LocalCenter += (-CardExtentZ + CardOffset) * Direction.AxisZ;
-
-		// set up the box for the card
-		OBB.Extent = (FVector3f)FVector3d(ExtentX, ExtentY, CardExtentZ + CardOffset*0.5);
-		OBB.Origin = (FVector3f)LocalCenter;
-
-		CardData.CardBuildData[CardIndex].OBB = OBB;
-		CardData.CardBuildData[CardIndex].AxisAlignedDirectionIndex = Direction.DirectionIndex;
-	}
-
-	bMeshCardsValid = true;
+	MeshCardRepresentation::SetCardsFromBounds(CardData);
 }
 
 
diff --git a/Engine/Source/Runtime/GeometryFramework/Public/Components/BaseDynamicMeshSceneProxy.h b/Engine/Source/Runtime/GeometryFramework/Public/Components/BaseDynamicMeshSceneProxy.h
--- a/Engine/Source/Runtime/GeometryFramework/Public/Components/BaseDynamicMeshSceneProxy.h
+++ b/Engine/Source/Runtime/GeometryFramework/Public/Components/BaseDynamicMeshSceneProxy.h
@@ -324,6 +324,8 @@ public:
 
 protected:
 	TPimplPtr<FCardRepresentationData> MeshCards;
+
+	UE_DEPRECATED(5.6, "Use MeshCards.IsValid() instead")
 	bool bMeshCardsValid = false;
 
 	// Helper to set lumen cards
