Made Dynamic Mesh Component's async physics build path more robust and disabled it in editor by default (w/ a cvar to experimentally re-enable it)

diff --git a/Engine/Source/Runtime/GeometryFramework/Private/Components/DynamicMeshComponent.cpp b/Engine/Source/Runtime/GeometryFramework/Private/Components/DynamicMeshComponent.cpp
--- a/Engine/Source/Runtime/GeometryFramework/Private/Components/DynamicMeshComponent.cpp
+++ b/Engine/Source/Runtime/GeometryFramework/Private/Components/DynamicMeshComponent.cpp
@@ -37,7 +37,11 @@ static TAutoConsoleVariable<int32> CVarDynamicMeshComponent_MaxComplexCollisionT
 	TEXT("If a DynamicMeshCompnent's UDynamicMesh has a larger triangle count than this value, it will not be passed to the Physics system to be used as Complex Collision geometry. A negative value indicates no limit.")
 );
 
-
+static TAutoConsoleVariable<bool> CVarDynamicMeshComponent_AllowCollisionAsyncBuildInEditor(
+	TEXT("geometry.DynamicMesh.AllowAsyncCollisionBuildInEditor"),
+	false,
+	TEXT("If enabled, we allow async build of dynamic mesh component collision in Editor. Otherwise, async collision build will only happen at runtime.")
+);
 
 
 namespace
@@ -1625,10 +1629,11 @@ void UDynamicMeshComponent::RebuildPhysicsData()
 		// It's uncertain whether these checks are actually needed. UProceduralMeshComponent checked for a game,
 		//  not editor, world. It's possible that at the time the editor was not ticked in a way that caused
 		//  physics builds to complete.
-		// We need to allow asynchronous builds in editor since dynamic meshes get used for real time modifications
-		//  that can otherwise hitch as things are dragged. We'll keep the guard against null or Inactive/None worlds
-		//  for now just in case, but it may not be necessary.
-		&& World && (World->IsGameWorld() || World->IsEditorWorld());
+		// We experimentally allow asynchronous builds in editor via a cvar, to support cases where in-editor dynamic meshes
+		//  may otherwise hitch during editing.
+		&& World && (World->IsGameWorld() || 
+			(World->IsEditorWorld() && CVarDynamicMeshComponent_AllowCollisionAsyncBuildInEditor.GetValueOnAnyThread())
+			);
 
 	UBodySetup* BodySetup = nullptr;
 	if (bUseAsyncCook)
@@ -1636,7 +1641,10 @@ void UDynamicMeshComponent::RebuildPhysicsData()
 		// Abort all previous ones still standing
 		for (UBodySetup* OldBody : AsyncBodySetupQueue)
 		{
-			OldBody->AbortPhysicsMeshAsyncCreation();
+			if (OldBody)
+			{
+				OldBody->AbortPhysicsMeshAsyncCreation();
+			}
 		}
 
 		BodySetup = CreateBodySetupHelper();
