Dynamic mesh component : Add support for two-sided rendering

diff --git a/Engine/Source/Runtime/GeometryFramework/Private/Components/BaseDynamicMeshComponent.cpp b/Engine/Source/Runtime/GeometryFramework/Private/Components/BaseDynamicMeshComponent.cpp
--- a/Engine/Source/Runtime/GeometryFramework/Private/Components/BaseDynamicMeshComponent.cpp
+++ b/Engine/Source/Runtime/GeometryFramework/Private/Components/BaseDynamicMeshComponent.cpp
@@ -201,9 +201,14 @@ void UBaseDynamicMeshComponent::SetVertexColorSpaceTransformMode(EDynamicMeshVer
 	}
 }
 
-
-
-
+void UBaseDynamicMeshComponent::SetTwoSided(bool bEnable)
+{
+	if (bTwoSided != bEnable)
+	{
+		bTwoSided = bEnable;
+		OnRenderingStateChanged(false);
+	}
+}
 
 void UBaseDynamicMeshComponent::SetEnableFlatShading(bool bEnable)
 {
diff --git a/Engine/Source/Runtime/GeometryFramework/Private/Components/BaseDynamicMeshSceneProxy.cpp b/Engine/Source/Runtime/GeometryFramework/Private/Components/BaseDynamicMeshSceneProxy.cpp
--- a/Engine/Source/Runtime/GeometryFramework/Private/Components/BaseDynamicMeshSceneProxy.cpp
+++ b/Engine/Source/Runtime/GeometryFramework/Private/Components/BaseDynamicMeshSceneProxy.cpp
@@ -42,6 +42,7 @@ namespace UE::DynamicMesh
 FBaseDynamicMeshSceneProxy::FBaseDynamicMeshSceneProxy(UBaseDynamicMeshComponent* Component)
 	: FPrimitiveSceneProxy(Component),
 	ParentBaseComponent(Component),
+	bTwoSided(Component->GetTwoSided()),
 	bEnableRaytracing(Component->GetEnableRaytracing()),
 	bEnableViewModeOverrides(Component->GetViewModeOverridesEnabled()),
 	bPreferStaticDrawPath(Component->GetMeshDrawPath() == EDynamicMeshDrawPath::StaticDraw)
@@ -468,7 +469,7 @@ void FBaseDynamicMeshSceneProxy::DrawBatch(FMeshElementCollector& Collector, con
 	FMeshBatchElement& BatchElement = Mesh.Elements[0];
 	BatchElement.IndexBuffer = &IndexBuffer;
 	Mesh.bWireframe = bWireframe;
-	//Mesh.bDisableBackfaceCulling = bWireframe;		// todo: doing this would be more consistent w/ other meshes in wireframe mode, but it is problematic for modeling tools - perhaps should be configurable
+	Mesh.bDisableBackfaceCulling = bTwoSided;
 	Mesh.VertexFactory = &RenderBuffers.VertexFactory;
 	Mesh.MaterialRenderProxy = UseMaterial;
 
diff --git a/Engine/Source/Runtime/GeometryFramework/Public/Components/BaseDynamicMeshComponent.h b/Engine/Source/Runtime/GeometryFramework/Public/Components/BaseDynamicMeshComponent.h
--- a/Engine/Source/Runtime/GeometryFramework/Public/Components/BaseDynamicMeshComponent.h
+++ b/Engine/Source/Runtime/GeometryFramework/Public/Components/BaseDynamicMeshComponent.h
@@ -336,7 +336,15 @@ public:
 	virtual EDynamicMeshVertexColorTransformMode GetVertexColorSpaceTransformMode() const { return ColorSpaceMode; }
 
 
+	//===============================================================================================================
+	// Two-Sided rendering support. When enabled, triangles will be rendered whether they are facing the camera or not
+public:
+	GEOMETRYFRAMEWORK_API virtual void SetTwoSided(bool bEnable);
+	virtual bool GetTwoSided() const { return bTwoSided; }
 
+protected:
+	UPROPERTY(EditAnywhere, Category = "Dynamic Mesh Component|Rendering");
+	bool bTwoSided = false;
 
 	//===============================================================================================================
 	// Flat shading support. When enabled, per-triangle normals are computed automatically and used in place
diff --git a/Engine/Source/Runtime/GeometryFramework/Public/Components/BaseDynamicMeshSceneProxy.h b/Engine/Source/Runtime/GeometryFramework/Public/Components/BaseDynamicMeshSceneProxy.h
--- a/Engine/Source/Runtime/GeometryFramework/Public/Components/BaseDynamicMeshSceneProxy.h
+++ b/Engine/Source/Runtime/GeometryFramework/Public/Components/BaseDynamicMeshSceneProxy.h
@@ -57,6 +57,9 @@ protected:
 	// use to control access to AllocatedBufferSets 
 	FCriticalSection AllocatedSetsLock;
 
+	// control whether the mesh is rendered two-sided
+	bool bTwoSided = false;
+
 	// control raytracing support
 	bool bEnableRaytracing = false;
 
