make dynamic mesh pools safer to use from a non-game thread

diff --git a/Engine/Source/Runtime/GeometryFramework/Private/UDynamicMesh.cpp b/Engine/Source/Runtime/GeometryFramework/Private/UDynamicMesh.cpp
--- a/Engine/Source/Runtime/GeometryFramework/Private/UDynamicMesh.cpp
+++ b/Engine/Source/Runtime/GeometryFramework/Private/UDynamicMesh.cpp
@@ -593,6 +593,7 @@ static TAutoConsoleVariable<int32> CVarDynamicMeshPoolMaxPoolSizeThreshold(
 
 UDynamicMesh* UDynamicMeshPool::RequestMesh()
 {
+	FScopeLock PoolScopeLock(&PoolAccessCs);
 	if (CachedMeshes.Num() > 0)
 	{
 		return CachedMeshes.Pop(EAllowShrinking::No);
@@ -601,11 +602,22 @@ UDynamicMesh* UDynamicMeshPool::RequestMesh()
 
 	// If we have allocated more meshes than our safety threshold, drop our holds on the existing meshes.
 	// This will allow them to be garbage-collected (eventually)
-	if (!ensure(AllCreatedMeshes.Num() < CVarDynamicMeshPoolMaxPoolSizeThreshold.GetValueOnGameThread()))
+	if (!ensure(AllCreatedMeshes.Num() < CVarDynamicMeshPoolMaxPoolSizeThreshold.GetValueOnAnyThread()))
 	{
-		UE_LOG(LogGeometry, Warning, TEXT("UDynamicMeshPool Threshold of %d Allocated Meshes exceeded! Releasing references to all current meshes and forcing a garbage collection."), CVarDynamicMeshPoolMaxPoolSizeThreshold.GetValueOnGameThread());
+		UE_LOG(LogGeometry, Warning, TEXT("UDynamicMeshPool Threshold of %d Allocated Meshes exceeded! Releasing references to all current meshes and forcing a garbage collection."), CVarDynamicMeshPoolMaxPoolSizeThreshold.GetValueOnAnyThread());
 		AllCreatedMeshes.Reset();
-		GEngine->ForceGarbageCollection(true);
+		// We generally expect the mesh pool to be used from game-thread blueprint code, but in case it is not, take extra care that ForceGarbageCollection is only called from the game thread
+		if (IsInGameThread())
+		{
+			GEngine->ForceGarbageCollection(true);
+		}
+		else
+		{
+			AsyncTask(ENamedThreads::GameThread, []()
+			{
+				GEngine->ForceGarbageCollection(true);
+			});
+		}
 	}
 
 	AllCreatedMeshes.Add(NewMesh);
@@ -616,6 +628,7 @@ UDynamicMesh* UDynamicMeshPool::RequestMesh()
 
 void UDynamicMeshPool::ReturnMesh(UDynamicMesh* Mesh)
 {
+	FScopeLock PoolScopeLock(&PoolAccessCs);
 	if ( ensure(Mesh) && ensure(AllCreatedMeshes.Contains(Mesh)) )
 	{
 		Mesh->Reset();
@@ -629,6 +642,7 @@ void UDynamicMeshPool::ReturnMesh(UDynamicMesh* Mesh)
 
 void UDynamicMeshPool::ReturnAllMeshes()
 {
+	FScopeLock PoolScopeLock(&PoolAccessCs);
 	CachedMeshes = AllCreatedMeshes;
 	for (UDynamicMesh* Mesh : CachedMeshes)
 	{
@@ -645,6 +659,7 @@ void UDynamicMeshPool::ReturnAllMeshes()
 
 void UDynamicMeshPool::FreeAllMeshes()
 {
+	FScopeLock PoolScopeLock(&PoolAccessCs);
 	CachedMeshes.Reset();
 	AllCreatedMeshes.Reset();
 }
diff --git a/Engine/Source/Runtime/GeometryFramework/Public/UDynamicMesh.h b/Engine/Source/Runtime/GeometryFramework/Public/UDynamicMesh.h
--- a/Engine/Source/Runtime/GeometryFramework/Public/UDynamicMesh.h
+++ b/Engine/Source/Runtime/GeometryFramework/Public/UDynamicMesh.h
@@ -412,4 +412,9 @@ protected:
 	/** All meshes the pool has allocated */
 	UPROPERTY()
 	TArray<TObjectPtr<UDynamicMesh>> AllCreatedMeshes;
+
+private:
+	// Typically the mesh pool is used from the game thread, but 
+	// in case it is not, we protect pool access w/ this critical section
+	mutable FCriticalSection PoolAccessCs;
 };
