clarify that the async physics cvar on dynamic mesh component is only about non-PIE editor cases

diff --git a/Engine/Source/Runtime/GeometryFramework/Private/Components/DynamicMeshComponent.cpp b/Engine/Source/Runtime/GeometryFramework/Private/Components/DynamicMeshComponent.cpp
--- a/Engine/Source/Runtime/GeometryFramework/Private/Components/DynamicMeshComponent.cpp
+++ b/Engine/Source/Runtime/GeometryFramework/Private/Components/DynamicMeshComponent.cpp
@@ -40,7 +40,7 @@ static TAutoConsoleVariable<int32> CVarDynamicMeshComponent_MaxComplexCollisionT
 static TAutoConsoleVariable<bool> CVarDynamicMeshComponent_AllowCollisionAsyncBuildInEditor(
 	TEXT("geometry.DynamicMesh.AllowAsyncCollisionBuildInEditor"),
 	false,
-	TEXT("If enabled, we allow async build of dynamic mesh component collision in Editor. Otherwise, async collision build will only happen at runtime.")
+	TEXT("If enabled, we allow async build of dynamic mesh component collision in (non-PIE) Editor. Otherwise, async collision build will only happen in game worlds (PIE/runtime).")
 );
 
 
