Dynamic Mesh Component: only call FlushRenderingCommands if a SceneProxy has been created.

diff --git a/Engine/Source/Runtime/GeometryFramework/Private/Components/BaseDynamicMeshComponent.cpp b/Engine/Source/Runtime/GeometryFramework/Private/Components/BaseDynamicMeshComponent.cpp
--- a/Engine/Source/Runtime/GeometryFramework/Private/Components/BaseDynamicMeshComponent.cpp
+++ b/Engine/Source/Runtime/GeometryFramework/Private/Components/BaseDynamicMeshComponent.cpp
@@ -224,8 +224,11 @@ void UBaseDynamicMeshComponent::OnRenderingStateChanged(bool bForceImmedateRebui
 {
 	if (bForceImmedateRebuild)
 	{
-		// finish any drawing so that we can be certain our SceneProxy is no longer in use before we rebuild it below
-		FlushRenderingCommands();
+		if (GetSceneProxy())
+		{
+			// finish any drawing so that we can be certain our SceneProxy is no longer in use before we rebuild it below
+			FlushRenderingCommands();
+		}
 
 		// force immediate rebuild of the SceneProxy
 		if (IsRegistered())
