[GeometryFramework] Update the Build.cs file to only add editor dependencies if we're building against the editor

#rnx

diff --git a/Engine/Source/Runtime/GeometryFramework/GeometryFramework.Build.cs b/Engine/Source/Runtime/GeometryFramework/GeometryFramework.Build.cs
--- a/Engine/Source/Runtime/GeometryFramework/GeometryFramework.Build.cs
+++ b/Engine/Source/Runtime/GeometryFramework/GeometryFramework.Build.cs
@@ -3,7 +3,7 @@
 using UnrealBuildTool;
 
 public class GeometryFramework : ModuleRules
-{	
+{
 	public GeometryFramework(ReadOnlyTargetRules Target) : base(Target)
 	{
 		//PCHUsage = ModuleRules.PCHUsageMode.UseExplicitOrSharedPCHs;
@@ -13,9 +13,9 @@ public class GeometryFramework : ModuleRules
 				"Shaders" // this is required for raytracing suppport?
 			}
 		);
-		
+
 		// These are required to register filters to the Level Editor Outliner
-		if (Target.Type == TargetType.Editor || Target.Type == TargetType.Program)
+		if (Target.bCompileAgainstEditor)
 		{
 			PrivateDependencyModuleNames.AddRange(
 				new string[]
@@ -24,10 +24,10 @@ public class GeometryFramework : ModuleRules
 					"Slate"
 				}
 			);
-				
+
 			PrivateIncludePathModuleNames.AddRange(
 				new string[]
-				{													
+				{
 					"LevelEditor",
 				}
 			);
