Add vertex-based region of interest update paths to DynamicMeshComponent, matching the Triangle ROI paths
+ expose these update paths to UPreviewMesh also

#Synced-CL 52243346

diff --git a/Engine/Plugins/Runtime/MeshModelingToolset/Source/ModelingComponents/Private/PreviewMesh.cpp b/Engine/Plugins/Runtime/MeshModelingToolset/Source/ModelingComponents/Private/PreviewMesh.cpp
--- a/Engine/Plugins/Runtime/MeshModelingToolset/Source/ModelingComponents/Private/PreviewMesh.cpp
+++ b/Engine/Plugins/Runtime/MeshModelingToolset/Source/ModelingComponents/Private/PreviewMesh.cpp
@@ -717,6 +717,12 @@ void UPreviewMesh::NotifyRegionDeferredEditCompleted(const TSet<int32>& Triangle
 	NotifyWorldPathTracedOutputInvalidated();
 }
 
+void UPreviewMesh::NotifyVertexRegionDeferredEditCompleted(const TArray<int32>& Vertices, EMeshRenderAttributeFlags ModifiedAttribs)
+{
+	DynamicMeshComponent->FastNotifyVerticesUpdated(Vertices, ModifiedAttribs);
+	NotifyWorldPathTracedOutputInvalidated();
+}
+
 
 void UPreviewMesh::NotifyWorldPathTracedOutputInvalidated()
 {
diff --git a/Engine/Plugins/Runtime/MeshModelingToolset/Source/ModelingComponents/Public/PreviewMesh.h b/Engine/Plugins/Runtime/MeshModelingToolset/Source/ModelingComponents/Public/PreviewMesh.h
--- a/Engine/Plugins/Runtime/MeshModelingToolset/Source/ModelingComponents/Public/PreviewMesh.h
+++ b/Engine/Plugins/Runtime/MeshModelingToolset/Source/ModelingComponents/Public/PreviewMesh.h
@@ -399,6 +399,13 @@ public:
 	 * @param ModifiedAttribs which mesh attributes have been modified and need to be updated
 	 */
 	UE_API void NotifyRegionDeferredEditCompleted(const TArray<int32>& Triangles, EMeshRenderAttributeFlags ModifiedAttribs);
+	
+	/**
+	 * Notify that a deferred edit is completed and cause update of rendering data structures for modified Vertices.
+	 * This can reduce the cost of mesh updates, but only if SetEnableRenderMeshDecomposition(true) has been called
+	 * @param ModifiedAttribs which mesh attributes have been modified and need to be updated
+	 */
+	UE_API void NotifyVertexRegionDeferredEditCompleted(const TArray<int32>& Vertices, EMeshRenderAttributeFlags ModifiedAttribs);
 
 	/**
 	 * Notify that a deferred edit is completed and cause update of rendering data structures for modified Triangles.
diff --git a/Engine/Source/Runtime/GeometryFramework/Private/Components/DynamicMeshComponent.cpp b/Engine/Source/Runtime/GeometryFramework/Private/Components/DynamicMeshComponent.cpp
--- a/Engine/Source/Runtime/GeometryFramework/Private/Components/DynamicMeshComponent.cpp
+++ b/Engine/Source/Runtime/GeometryFramework/Private/Components/DynamicMeshComponent.cpp
@@ -21,6 +21,7 @@
 
 #include "UObject/UE5ReleaseStreamObjectVersion.h"
 #include "UObject/UObjectGlobals.h"
+#include <atomic>
 
 // default proxy for this component
 #include "Components/DynamicMeshSceneProxy.h"
@@ -673,92 +674,74 @@ void UDynamicMeshComponent::FastNotifySecondaryTrianglesChanged()
 
 void UDynamicMeshComponent::FastNotifyTriangleVerticesUpdated(const TArray<int32>& Triangles, EMeshRenderAttributeFlags UpdatedAttributes)
 {
-	if (!ensure(MeshObject))
+	FastNotifyDecompositionSetsUpdated([this, &Triangles]() -> TArray<int32>
 	{
-		return;
-	}
-
-	// should not be using fast paths if we have to run mesh postprocessor
-	if (ensure(!RenderMeshPostProcessor) == false)
-	{
-		RenderMeshPostProcessor->ProcessMesh(*GetMesh(), *RenderMesh);
-		ResetProxy();
-		return;
-	}
+		TRACE_CPUPROFILER_EVENT_SCOPE(SimpleDynamicMeshComponent_FastTriangleVertexUpdate_FindSets);
+		// compute list of sets to update
+		TArray<int32> UpdatedSets;	
+		for (int32 tid : Triangles)
+		{
+			int32 SetID = Decomposition->GetGroupForTriangle(tid);
+			UpdatedSets.AddUnique(SetID);
+		}
+		return UpdatedSets;
+	}, UpdatedAttributes);
+}
 
-	bool bUpdateSecondarySort = (SecondaryTriFilterFunc) &&
-		((UpdatedAttributes & EMeshRenderAttributeFlags::SecondaryIndexBuffers) != EMeshRenderAttributeFlags::None);
 
-	FDynamicMeshSceneProxy* Proxy = GetCurrentSceneProxy();
-	if (!Proxy || !AllowFastUpdate())
-	{
-		ResetProxy();
-	}
-	else if ( ! Decomposition )
+void UDynamicMeshComponent::FastNotifyTriangleVerticesUpdated(const TSet<int32>& Triangles, EMeshRenderAttributeFlags UpdatedAttributes)
+{
+	FastNotifyDecompositionSetsUpdated([this, &Triangles]() -> TArray<int32>
 	{
-		FastNotifyVertexAttributesUpdated(UpdatedAttributes);
-		if (bUpdateSecondarySort)
+		TRACE_CPUPROFILER_EVENT_SCOPE(SimpleDynamicMeshComponent_FastTriangleVertexUpdate_FindSets);
+		// compute list of sets to update
+		TArray<int32> UpdatedSets;
+		for (int32 tid : Triangles)
 		{
-			Proxy->FastUpdateAllIndexBuffers();
+			int32 SetID = Decomposition->GetGroupForTriangle(tid);
+			UpdatedSets.AddUnique(SetID);
 		}
-		GetDynamicMesh()->PostRealtimeUpdate();
-	}
-	else
+		return UpdatedSets;
+	}, UpdatedAttributes);
+}
+
+void UDynamicMeshComponent::FastNotifyVerticesUpdated(const TArray<int32>& Vertices, EMeshRenderAttributeFlags UpdatedAttributes)
+{
+	FastNotifyDecompositionSetsUpdated([this, &Vertices]() -> TArray<int32>
 	{
+		TRACE_CPUPROFILER_EVENT_SCOPE(SimpleDynamicMeshComponent_FastVertexUpdate_FindSets);
 		// compute list of sets to update
-		TArray<int32> UpdatedSets;
+		TArray<int32> UpdatedBuffers;
+		int32 NumBuffers = Decomposition->Num();
+		TArray<std::atomic<bool>> BufferFlags;
+		BufferFlags.SetNum(NumBuffers);
+		for (int32 k = 0; k < NumBuffers; ++k)
 		{
-			TRACE_CPUPROFILER_EVENT_SCOPE(SimpleDynamicMeshComponent_FastVertexUpdate_FindSets);
-			for (int32 tid : Triangles)
-			{
-				int32 SetID = Decomposition->GetGroupForTriangle(tid);
-				UpdatedSets.AddUnique(SetID);
-			}
+			BufferFlags[k] = false;
 		}
-
-		bool bPositions = (UpdatedAttributes & EMeshRenderAttributeFlags::Positions) != EMeshRenderAttributeFlags::None;
-
-		// calculate bounds while we are updating vertices
-		TFuture<void> UpdateBoundsCalc;
-		if (bPositions)
+		const FDynamicMesh3* Mesh = GetMesh();
+		check(Mesh);
+		ParallelFor(Vertices.Num(), [&](int32 k)
 		{
-			UpdateBoundsCalc = Async(DynamicMeshComponentAsyncExecTarget, [this]()
+			int32 VID = Vertices[k];
+			Mesh->EnumerateVertexTriangles(VID, [this, &BufferFlags](int32 TID)
 			{
-				TRACE_CPUPROFILER_EVENT_SCOPE(SimpleDynamicMeshComponent_FastVertexUpdate_AsyncBoundsUpdate);
-				UpdateLocalBounds();
+				int32 SetID = Decomposition->GetGroupForTriangle(TID);
+				BufferFlags[SetID].store(true, std::memory_order_relaxed);
 			});
-		}
-
-		// update the render buffers
-		{
-			TRACE_CPUPROFILER_EVENT_SCOPE(SimpleDynamicMeshComponent_FastVertexUpdate_ApplyUpdate);
-			Proxy->FastUpdateVertices(UpdatedSets, bPositions,
-				(UpdatedAttributes & EMeshRenderAttributeFlags::VertexNormals) != EMeshRenderAttributeFlags::None,
-				(UpdatedAttributes & EMeshRenderAttributeFlags::VertexColors) != EMeshRenderAttributeFlags::None,
-				(UpdatedAttributes & EMeshRenderAttributeFlags::VertexUVs) != EMeshRenderAttributeFlags::None);
-		}
-
-		if (bUpdateSecondarySort)
-		{
-			Proxy->FastUpdateIndexBuffers(UpdatedSets);
-		}
-
-		if (bPositions)
+		});
+		for (int32 k = 0; k < NumBuffers; ++k)
 		{
-			TRACE_CPUPROFILER_EVENT_SCOPE(SimpleDynamicMeshComponent_FastVertexUpdate_FinalPositionsUpdate);
-			MarkRenderTransformDirty();
-			UpdateBoundsCalc.Wait();
-			UpdateBounds();
+			if (BufferFlags[k])
+			{
+				UpdatedBuffers.Add(k);
+			}
 		}
-
-		GetDynamicMesh()->PostRealtimeUpdate();
-	}
+		return UpdatedBuffers;
+	}, UpdatedAttributes);
 }
 
-
-
-
-void UDynamicMeshComponent::FastNotifyTriangleVerticesUpdated(const TSet<int32>& Triangles, EMeshRenderAttributeFlags UpdatedAttributes)
+void UDynamicMeshComponent::FastNotifyDecompositionSetsUpdated(TFunctionRef<TArray<int32>(void)> GetUpdatedDecompositionIndices, EMeshRenderAttributeFlags UpdatedAttributes)
 {
 	if (!ensure(MeshObject))
 	{
@@ -793,15 +776,7 @@ void UDynamicMeshComponent::FastNotifyTriangleVerticesUpdated(const TSet<int32>&
 	else
 	{
 		// compute list of sets to update
-		TArray<int32> UpdatedSets;
-		{
-			TRACE_CPUPROFILER_EVENT_SCOPE(SimpleDynamicMeshComponent_FastVertexUpdate_FindSets);
-			for (int32 tid : Triangles)
-			{
-				int32 SetID = Decomposition->GetGroupForTriangle(tid);
-				UpdatedSets.AddUnique(SetID);
-			}
-		}
+		TArray<int32> UpdatedSets = GetUpdatedDecompositionIndices();
 
 		bool bPositions = (UpdatedAttributes & EMeshRenderAttributeFlags::Positions) != EMeshRenderAttributeFlags::None;
 
@@ -811,32 +786,29 @@ void UDynamicMeshComponent::FastNotifyTriangleVerticesUpdated(const TSet<int32>&
 		{
 			UpdateBoundsCalc = Async(DynamicMeshComponentAsyncExecTarget, [this]()
 			{
-				TRACE_CPUPROFILER_EVENT_SCOPE(SimpleDynamicMeshComponent_FastVertexUpdate_AsyncBoundsUpdate);
+				TRACE_CPUPROFILER_EVENT_SCOPE(SimpleDynamicMeshComponent_FastROIUpdate_AsyncBoundsUpdate);
 				UpdateLocalBounds();
 			});
 		}
 
 		// update the render buffers
 		{
-			TRACE_CPUPROFILER_EVENT_SCOPE(SimpleDynamicMeshComponent_FastVertexUpdate_ApplyUpdate);
+			TRACE_CPUPROFILER_EVENT_SCOPE(SimpleDynamicMeshComponent_FastROIUpdate_ApplyUpdate);
 			Proxy->FastUpdateVertices(UpdatedSets, bPositions,
 				(UpdatedAttributes & EMeshRenderAttributeFlags::VertexNormals) != EMeshRenderAttributeFlags::None,
 				(UpdatedAttributes & EMeshRenderAttributeFlags::VertexColors) != EMeshRenderAttributeFlags::None,
 				(UpdatedAttributes & EMeshRenderAttributeFlags::VertexUVs) != EMeshRenderAttributeFlags::None);
 		}
 
+		if (bUpdateSecondarySort)
 		{
-			TRACE_CPUPROFILER_EVENT_SCOPE(SimpleDynamicMeshComponent_FastVertexUpdate_UpdateIndexBuffers);
-			if (bUpdateSecondarySort)
-			{
-				Proxy->FastUpdateIndexBuffers(UpdatedSets);
-			}
+			Proxy->FastUpdateIndexBuffers(UpdatedSets);
 		}
 
 		// finish up, have to wait for background bounds recalculation here
 		if (bPositions)
 		{
-			TRACE_CPUPROFILER_EVENT_SCOPE(SimpleDynamicMeshComponent_FastVertexUpdate_FinalPositionsUpdate);
+			TRACE_CPUPROFILER_EVENT_SCOPE(SimpleDynamicMeshComponent_FastROIUpdate_FinalPositionsUpdate);
 			MarkRenderTransformDirty();
 			UpdateBoundsCalc.Wait();
 			UpdateBounds();
diff --git a/Engine/Source/Runtime/GeometryFramework/Public/Components/DynamicMeshComponent.h b/Engine/Source/Runtime/GeometryFramework/Public/Components/DynamicMeshComponent.h
--- a/Engine/Source/Runtime/GeometryFramework/Public/Components/DynamicMeshComponent.h
+++ b/Engine/Source/Runtime/GeometryFramework/Public/Components/DynamicMeshComponent.h
@@ -320,6 +320,17 @@ public:
 	 */
 	GEOMETRYFRAMEWORK_API void FastNotifyTriangleVerticesUpdated(const TSet<int32>& Triangles, EMeshRenderAttributeFlags UpdatedAttributes);
 
+	/**
+	 * This function updates vertex positions/attributes of existing SceneProxy render buffers if possible, for the given vertices.
+	 * If a FMeshRenderDecomposition has not been explicitly set, call is forwarded to FastNotifyVertexAttributesUpdated()
+	 */
+	GEOMETRYFRAMEWORK_API void FastNotifyVerticesUpdated(const TArray<int32>& Vertices, EMeshRenderAttributeFlags UpdatedAttributes);
+
+private:
+	// Helper to implement common functionality of above FastNotify-region-Updated methods
+	void FastNotifyDecompositionSetsUpdated(TFunctionRef<TArray<int32>(void)> GetUpdatedDecompositionIndices, EMeshRenderAttributeFlags UpdatedAttributes);
+public:
+
 	/**
 	 * If a Decomposition is set on this Component, and everything is currently valid (proxy/etc), precompute the set of
 	 * buffers that will be modified, as well as the bounds of the modified region. These are both computed in parallel.
